Fix error in metadata name

diff --git a/Engine/Source/Editor/StatsViewer/Classes/TextureStats.h b/Engine/Source/Editor/StatsViewer/Classes/TextureStats.h
--- a/Engine/Source/Editor/StatsViewer/Classes/TextureStats.h
+++ b/Engine/Source/Editor/StatsViewer/Classes/TextureStats.h
@@ -21,7 +21,7 @@ enum ETextureObjectSets
 };
 
 /** Statistics page for textures. */
-UCLASS(Transient, MinimalAPI, meta=( DisplayName = "Texture Stats", TObjectSetType = "ETextureObjectSets" ) )
+UCLASS(Transient, MinimalAPI, meta=( DisplayName = "Texture Stats", ObjectSetType = "ETextureObjectSets" ) )
 class UTextureStats : public UObject
 {
 	GENERATED_UCLASS_BODY()
