Removed view stat options from Build menu and added config entry to remember the last stat page that was viewed (defaults to primitive stats)

diff --git a/Engine/Source/Editor/LevelEditor/Private/LevelEditorToolBar.cpp b/Engine/Source/Editor/LevelEditor/Private/LevelEditorToolBar.cpp
--- a/Engine/Source/Editor/LevelEditor/Private/LevelEditorToolBar.cpp
+++ b/Engine/Source/Editor/LevelEditor/Private/LevelEditorToolBar.cpp
@@ -446,15 +446,6 @@ TSharedRef< SWidget > FLevelEditorToolBar::GenerateBuildMenuContent( TSharedRef<
 	}
 	MenuBuilder.EndSection();
 
-
-	// Texture Stats
-	MenuBuilder.BeginSection("LevelEditorStatistics", LOCTEXT( "Statistics", "Statistics" ) );
-	{
-	    MenuBuilder.AddMenuEntry( FLevelEditorCommands::Get().SceneStats, NAME_None, LOCTEXT("OpenSceneStats", "Scene Stats") );
-	    MenuBuilder.AddMenuEntry( FLevelEditorCommands::Get().TextureStats, NAME_None, LOCTEXT("OpenTextureStats", "TextureStats Stats") );
-	}
-	MenuBuilder.EndSection();
-
 	// Map Check
 	MenuBuilder.BeginSection("LevelEditorVerification", LOCTEXT( "VerificationHeading", "Verification" ) );
 	{
diff --git a/Engine/Source/Editor/StatsViewer/Private/SStatsViewer.cpp b/Engine/Source/Editor/StatsViewer/Private/SStatsViewer.cpp
--- a/Engine/Source/Editor/StatsViewer/Private/SStatsViewer.cpp
+++ b/Engine/Source/Editor/StatsViewer/Private/SStatsViewer.cpp
@@ -24,6 +24,9 @@ namespace StatsViewerConstants
 {
 	/** Delay (in seconds) after a new character is entered into the search box to wait before updating the list (to give them time to enter a whole string instead of useless updating every time a char is put in) **/
 	static const float SearchTextUpdateDelay = 0.5f;
+
+	/** Stat viewer config file section name */
+	static const FString ConfigSectionName = "StatsViewer";
 }
 
 namespace StatsViewerMetadata
@@ -194,6 +197,24 @@ void SStatsViewer::Construct( const FArguments& InArgs )
 			]
 		]
 	];
+
+	// Display stats page from previous stat viewer instance
+	if (!CurrentStats.IsValid())
+	{
+		TSharedPtr<IStatsPage> InitialStatsPage;
+		FString DisplayedStatsPageName;
+		if (GConfig->GetString(*StatsViewerConstants::ConfigSectionName, TEXT("DisplayedStatsPageName"), DisplayedStatsPageName, GEditorUserSettingsIni))
+		{
+			InitialStatsPage = FStatsPageManager::Get().GetPage(FName(*DisplayedStatsPageName));
+		}
+		else
+		{
+			// Default to primitive stats if no config data exists yet
+			InitialStatsPage = FStatsPageManager::Get().GetPage(EStatsPage::PrimitiveStats);
+		}
+
+		SetDisplayedStats(InitialStatsPage.ToSharedRef());
+	}
 }
 
 SStatsViewer::SStatsViewer() :
@@ -684,6 +705,7 @@ void SStatsViewer::SetDisplayedStats( TSharedRef<IStatsPage> StatsPage )
 	}
 
 	CurrentStats = StatsPage;
+	GConfig->SetString(*StatsViewerConstants::ConfigSectionName, TEXT("DisplayedStatsPageName"), *StatsPage->GetName().ToString(), GEditorUserSettingsIni);
 	CurrentStats->OnShow( SharedThis(this) );
 	CurrentObjectSetIndex = CurrentStats->GetSelectedObjectSet();
 	CurrentFilterIndex = 0;
