StatsViewer: Added ability to override filter widget on stats pages

diff --git a/Engine/Source/Editor/StatsViewer/Private/SStatsViewer.cpp b/Engine/Source/Editor/StatsViewer/Private/SStatsViewer.cpp
--- a/Engine/Source/Editor/StatsViewer/Private/SStatsViewer.cpp
+++ b/Engine/Source/Editor/StatsViewer/Private/SStatsViewer.cpp
@@ -134,16 +134,9 @@ void SStatsViewer::Construct( const FArguments& InArgs )
 				.Padding(0.0f)
 				.HAlign(HAlign_Right)
 				[
-					SNew( SComboButton )
-					.Visibility( this, &SStatsViewer::OnGetObjectSetsVisibility )
-					.ContentPadding(3)
-					.OnGetMenuContent( this, &SStatsViewer::OnGetObjectSetMenuContent )
-					.ButtonContent()
-					[
-						SNew( STextBlock )
-						.Text( this, &SStatsViewer::OnGetObjectSetMenuLabel )
-						.ToolTipText( LOCTEXT( "DisplayedObjects_Tooltip", "Choose the objects whose statistics you want to display" ) )
-					]
+					SAssignNew( CustomFilter, SBorder )
+					.BorderImage( FEditorStyle::GetBrush("NoBorder") )
+					.Padding(0.0f)
 				]
 			]
 		]
@@ -712,15 +705,37 @@ void SStatsViewer::SetDisplayedStats( TSharedRef<IStatsPage> StatsPage )
 	FilterTextBoxWidget->SetText( FText::FromString(TEXT("")) );
 
 	// set custom widget, if any
-	TSharedPtr<SWidget> CustomWidget = CurrentStats->GetCustomWidget( SharedThis(this) );
-	if(CustomWidget.IsValid())
+	TSharedPtr<SWidget> CustomContentWidget = CurrentStats->GetCustomWidget( SharedThis(this) );
 	{
-		CustomContent->SetContent( CustomWidget.ToSharedRef() );
-		CustomContent->SetVisibility( EVisibility::Visible );
+		if(CustomContentWidget.IsValid())
+		{
+			CustomContent->SetContent( CustomContentWidget.ToSharedRef() );
+			CustomContent->SetVisibility( EVisibility::Visible );
+		}
+		else
+		{
+			CustomContent->SetVisibility( EVisibility::Collapsed );
+		}
 	}
-	else
+
+	// set custom filter, if any
+	TSharedPtr<SWidget> CustomFilterWidget = CurrentStats->GetCustomFilter( SharedThis(this) );
 	{
-		CustomContent->SetVisibility( EVisibility::Collapsed );
+		if (!CustomFilterWidget.IsValid())
+		{
+			CustomFilterWidget = SNew( SComboButton )
+				.Visibility( this, &SStatsViewer::OnGetObjectSetsVisibility )
+				.ContentPadding(3)
+				.OnGetMenuContent( this, &SStatsViewer::OnGetObjectSetMenuContent )
+				.ButtonContent()
+				[
+					SNew( STextBlock )
+						.Text( this, &SStatsViewer::OnGetObjectSetMenuLabel )
+						.ToolTipText( LOCTEXT( "DisplayedObjects_Tooltip", "Choose the objects whose statistics you want to display" ) )
+				];
+		}
+
+		CustomFilter->SetContent(CustomFilterWidget.ToSharedRef());
 	}
 
 	Refresh();
diff --git a/Engine/Source/Editor/StatsViewer/Private/SStatsViewer.h b/Engine/Source/Editor/StatsViewer/Private/SStatsViewer.h
--- a/Engine/Source/Editor/StatsViewer/Private/SStatsViewer.h
+++ b/Engine/Source/Editor/StatsViewer/Private/SStatsViewer.h
@@ -155,5 +155,8 @@ private:
 
 	/** Container for custom content supplied by stats pages */
 	TSharedPtr< SBorder > CustomContent;
+
+	/** Container for custom filters supplied by stats pages */
+	TSharedPtr< SBorder > CustomFilter;
 };
 
diff --git a/Engine/Source/Editor/StatsViewer/Public/IStatsPage.h b/Engine/Source/Editor/StatsViewer/Public/IStatsPage.h
--- a/Engine/Source/Editor/StatsViewer/Public/IStatsPage.h
+++ b/Engine/Source/Editor/StatsViewer/Public/IStatsPage.h
@@ -105,7 +105,13 @@ public:
 	virtual void GenerateTotals( const TArray< TWeakObjectPtr<UObject> >& InObjects, TMap<FString, FText>& OutTotals ) const = 0;
 
 	/** 
-	 * Get custom widget to display in the top part of the stats viewer, can return NULL 
+	 * Get custom filter to display in the top part of the stats viewer, can return nullptr 
+	 * @param	InParentStatsViewer	The parent stats viewer
+	 */
+	virtual TSharedPtr<SWidget> GetCustomFilter( TWeakPtr< class IStatsViewer > InParentStatsViewer ) = 0;
+
+	/** 
+	 * Get custom widget to display in the top part of the stats viewer, can return nullptr 
 	 * @param	InParentStatsViewer	The parent stats viewer
 	 */
 	virtual TSharedPtr<SWidget> GetCustomWidget( TWeakPtr< class IStatsViewer > InParentStatsViewer ) = 0;
diff --git a/Engine/Source/Editor/StatsViewer/Public/StatsPage.h b/Engine/Source/Editor/StatsViewer/Public/StatsPage.h
--- a/Engine/Source/Editor/StatsViewer/Public/StatsPage.h
+++ b/Engine/Source/Editor/StatsViewer/Public/StatsPage.h
@@ -21,7 +21,7 @@ public:
 		FString EnumName = Entry::StaticClass()->GetName();
 		EnumName += TEXT(".");
 		EnumName += Entry::StaticClass()->GetMetaData( TEXT("ObjectSetType") );
-		ObjectSetEnum = FindObject<UEnum>( NULL, *EnumName );
+		ObjectSetEnum = FindObject<UEnum>( nullptr, *EnumName );
 		bRefresh = false;
 		bShow = false;
 		ObjectSetIndex = 0;
@@ -65,7 +65,7 @@ public:
 
 	virtual int32 GetObjectSetCount() const override
 	{
-		if(ObjectSetEnum != NULL)
+		if(ObjectSetEnum != nullptr)
 		{
 			return ObjectSetEnum->NumEnums() - 1;
 		}
@@ -74,7 +74,7 @@ public:
 
 	virtual FString GetObjectSetName( int32 InObjectSetIndex ) const override
 	{
-		if(ObjectSetEnum != NULL)
+		if(ObjectSetEnum != nullptr)
 		{
 			return ObjectSetEnum->GetDisplayNameText( InObjectSetIndex ).ToString();
 		}
@@ -86,7 +86,7 @@ public:
 
 	virtual FString GetObjectSetToolTip( int32 InObjectSetIndex ) const override
 	{
-		if(ObjectSetEnum != NULL)
+		if(ObjectSetEnum != nullptr)
 		{
 			return ObjectSetEnum->GetToolTipText( InObjectSetIndex ).ToString();
 		}
@@ -101,9 +101,14 @@ public:
 		return Entry::StaticClass();
 	}
 
+	virtual TSharedPtr<SWidget> GetCustomFilter( TWeakPtr< class IStatsViewer > InParentStatsViewer ) override
+	{
+		return nullptr;
+	}
+
 	virtual TSharedPtr<SWidget> GetCustomWidget( TWeakPtr< class IStatsViewer > InParentStatsViewer ) override
 	{
-		return NULL;
+		return nullptr;
 	}
 
 	virtual void SetSelectedObjectSet( int32 InObjectSetIndex ) override
