StatsViewer: Added support for hyper-linked arrays of UObjects

diff --git a/Engine/Source/Editor/StatsViewer/Private/ActorArrayHyperlinkColumn.cpp b/Engine/Source/Editor/StatsViewer/Private/ActorArrayHyperlinkColumn.cpp
--- a/Engine/Source/Editor/StatsViewer/Private/ActorArrayHyperlinkColumn.cpp
+++ b/Engine/Source/Editor/StatsViewer/Private/ActorArrayHyperlinkColumn.cpp
@@ -27,8 +27,7 @@ public:
 			FirstElement->GetValue(Object);
 			if(Object != NULL)
 			{
-				AActor* Actor = CastChecked<AActor>(Object);
-				Text = FText::FromString( Actor->GetName() );
+				Text = FText::FromString( Object->GetName() );
 			}
 			else
 			{
@@ -71,6 +70,8 @@ public:
 		if (NumElements > 0)
 		{
 			TArray<AActor*> ActorsToFocus;
+			TArray<UObject*> ObjectsToSync;
+
 			const FScopedTransaction Transaction( LOCTEXT("SelectActors", "Statistics Select Actors") );
 			GEditor->SelectNone(false, false);
 			for(uint32 ElementIndex = 0; ElementIndex < NumElements; ElementIndex++)
@@ -80,9 +81,17 @@ public:
 				Element->GetValue(Object);
 				if(Object != NULL)
 				{
-					AActor* Actor = CastChecked<AActor>(Object);
-					GEditor->SelectActor(Actor, true, true, true);
-					ActorsToFocus.Add(Actor);
+					AActor* Actor = Cast<AActor>(Object);
+
+					if (Actor != nullptr)
+					{
+						GEditor->SelectActor(Actor, true, true, true);
+						ActorsToFocus.Add(Actor);
+					}
+					else
+					{
+						ObjectsToSync.Add(Object);
+					}
 				}
 			}
 
@@ -90,6 +99,11 @@ public:
 			{
 				GEditor->MoveViewportCamerasToActor(ActorsToFocus, false);	
 			}
+
+			if (ObjectsToSync.Num() > 0)
+			{
+				GEditor->SyncBrowserToObjects(ObjectsToSync);
+			}
 		}
 	}
 
@@ -113,11 +127,7 @@ bool FActorArrayHyperlinkColumn::Supports( const TSharedRef< IPropertyTableColum
 				UArrayProperty* ArrayProperty = Cast<UArrayProperty>(Property);
 				if( ArrayProperty->Inner->IsA(UWeakObjectProperty::StaticClass()) )
 				{
-					const UClass* PropertyClass = Cast<UWeakObjectProperty>(ArrayProperty->Inner)->PropertyClass;
-					if( PropertyClass == AActor::StaticClass() )
-					{
-						return true;
-					}
+					return true;
 				}
 			}
 		}
