StatsViewer: Implemented cooker statistics page

diff --git a/Engine/Source/Editor/StatsViewer/Classes/CookerStats.h b/Engine/Source/Editor/StatsViewer/Classes/CookerStats.h
new file mode 100644
--- /dev/null
+++ b/Engine/Source/Editor/StatsViewer/Classes/CookerStats.h
@@ -0,0 +1,42 @@
+// Copyright 1998-2014 Epic Games, Inc. All Rights Reserved.
+
+#pragma once
+
+#include "CookerStats.generated.h"
+
+
+/** Enum defining the object sets for this stats object */
+UENUM()
+enum ECookerStatsObjectSets
+{
+	CookerStatsObjectSets_Default UMETA(DisplayName="Default", ToolTip="View cooker statistics"),
+};
+
+
+/**
+ * Statistics for a cooked asset.
+ *
+ * Note: We assume that asset files are not larger than 2GB, because the StatsViewer is still lacking int64 support.
+ */
+UCLASS(Transient, MinimalAPI, meta=(DisplayName="Cooker Stats", ObjectSetType="ECookerStatsObjectSets"))
+class UCookerStats
+	: public UObject
+{
+	GENERATED_UCLASS_BODY()
+
+	/** The assets contained in the file. */
+	UPROPERTY(VisibleAnywhere, AssetRegistrySearchable, Category="Stats", meta=(DisplayName="Asset(s)", ColumnWidth="100"))
+	TArray<TWeakObjectPtr<UObject>> Assets;
+
+	/** The size of the assets before cooking. */
+	UPROPERTY(VisibleAnywhere, AssetRegistrySearchable, Category="Stats", meta=(DisplayName="Size (Original)", ColumnWidth="50", ShowTotal="true", Unit="KB"))
+	float SizeBefore;
+
+	/** The size of the assets after cooking. */
+	UPROPERTY(VisibleAnywhere, AssetRegistrySearchable, Category="Stats", meta=(DisplayName="Size (Cooked)", ColumnWidth="50", ShowTotal="true", Unit="KB"))
+	float SizeAfter;
+
+	/** Asset path without the name "package.[group.]" */
+	UPROPERTY(VisibleAnywhere, AssetRegistrySearchable, Category="Stats", meta=(ColumnWidth="300"))
+	FString Path;
+};
diff --git a/Engine/Source/Editor/StatsViewer/Private/StatsEntries/CookerStats.cpp b/Engine/Source/Editor/StatsViewer/Private/StatsEntries/CookerStats.cpp
new file mode 100644
--- /dev/null
+++ b/Engine/Source/Editor/StatsViewer/Private/StatsEntries/CookerStats.cpp
@@ -0,0 +1,9 @@
+// Copyright 1998-2014 Epic Games, Inc. All Rights Reserved.
+
+#include "StatsViewerPrivatePCH.h"
+#include "CookerStats.h"
+
+
+UCookerStats::UCookerStats( const class FPostConstructInitializeProperties& PCIP )
+	: Super(PCIP)
+{ }
diff --git a/Engine/Source/Editor/StatsViewer/Private/StatsPages/CookerStatsPage.cpp b/Engine/Source/Editor/StatsViewer/Private/StatsPages/CookerStatsPage.cpp
new file mode 100644
--- /dev/null
+++ b/Engine/Source/Editor/StatsViewer/Private/StatsPages/CookerStatsPage.cpp
@@ -0,0 +1,236 @@
+// Copyright 1998-2014 Epic Games, Inc. All Rights Reserved.
+
+#include "StatsViewerPrivatePCH.h"
+#include "AssetRegistryModule.h"
+#include "CookerStatsPage.h"
+
+
+#define LOCTEXT_NAMESPACE "FCookerStatsPage"
+
+
+/* FCookerStatsPage static functions
+ *****************************************************************************/
+
+FCookerStatsPage& FCookerStatsPage::Get()
+{
+	static FCookerStatsPage* Instance = nullptr;
+
+	if (Instance == nullptr)
+	{
+		Instance = new FCookerStatsPage;
+	}
+
+	return *Instance;
+}
+
+
+/* IStatsPage interface
+ *****************************************************************************/
+
+void FCookerStatsPage::Generate( TArray<TWeakObjectPtr<UObject>>& OutObjects ) const
+{
+	// visitor for the directory enumerator
+	class FPlatformDirectoryVisitor
+		: public IPlatformFile::FDirectoryVisitor
+	{
+	public:
+
+		FPlatformDirectoryVisitor( IAssetRegistry& InAssetRegistry, TArray<TWeakObjectPtr<UObject>>& InOutObjects, const FString& InPlatformName )
+			: AssetRegistry(InAssetRegistry)
+			, OutObjects(InOutObjects)
+			, PlatformName(InPlatformName)
+		{ }
+
+		virtual bool Visit( const TCHAR* FilenameOrDirectory, bool bIsDirectory ) override
+		{
+			if (!bIsDirectory)
+			{
+				FString OriginalFileName = FString(FilenameOrDirectory).Replace(*(FPaths::GameSavedDir() / TEXT("Cooked") / PlatformName), *FPaths::RootDir());
+
+				if (OriginalFileName.EndsWith(TEXT(".uasset")) || OriginalFileName.EndsWith(TEXT(".umap")))
+				{
+					UCookerStats* Entry = ConstructObject<UCookerStats>(UCookerStats::StaticClass());
+					Entry->AddToRoot();
+					Entry->Path = FilenameOrDirectory;
+					Entry->SizeAfter = (int32)IFileManager::Get().FileSize(FilenameOrDirectory) / 1024.0f;
+					Entry->SizeBefore = FMath::Max(0.0f, (int32)IFileManager::Get().FileSize(*OriginalFileName) / 1024.0f);
+
+					FString PackageName;
+
+					if (FPackageName::TryConvertFilenameToLongPackageName(OriginalFileName, PackageName))
+					{
+						TArray<FAssetData> Assets;
+						AssetRegistry.GetAssetsByPackageName(FName(*PackageName), Assets);
+
+						for (auto Asset : Assets)
+						{
+							Entry->Assets.Add(Asset.GetAsset());
+						}
+					}
+
+					OutObjects.Add(Entry);
+				}
+			}
+
+			return true;
+		}
+
+	private:
+
+		IAssetRegistry& AssetRegistry;
+		TArray<TWeakObjectPtr<UObject>>& OutObjects;
+		FString PlatformName;
+	};
+
+
+	if (SelectedPlatformName.IsEmpty())
+	{
+		return;
+	}
+
+	// create data objects for each cooked asset
+	FAssetRegistryModule& AssetRegistryModule = FModuleManager::LoadModuleChecked<FAssetRegistryModule>("AssetRegistry");
+	FPlatformDirectoryVisitor PlatformDirectoryVisitor(AssetRegistryModule.Get(), OutObjects, SelectedPlatformName);
+	FPlatformFileManager::Get().GetPlatformFile().IterateDirectoryRecursively(*(FPaths::GameSavedDir() / TEXT("Cooked") / SelectedPlatformName), PlatformDirectoryVisitor);
+}
+
+
+void FCookerStatsPage::GenerateTotals( const TArray<TWeakObjectPtr<UObject>>& InObjects, TMap<FString, FText>& OutTotals ) const
+{
+	if (InObjects.Num() == 0)
+	{
+		return;
+	}
+
+	UCookerStats* TotalEntry = ConstructObject<UCookerStats>(UCookerStats::StaticClass());
+	{
+		for (auto It = InObjects.CreateConstIterator(); It; ++It)
+		{
+			UCookerStats* StatsEntry = Cast<UCookerStats>(It->Get());
+			
+			TotalEntry->SizeBefore += StatsEntry->SizeBefore;
+			TotalEntry->SizeAfter += StatsEntry->SizeAfter;
+		}
+
+		OutTotals.Add(TEXT("SizeBefore"), FText::AsNumber(TotalEntry->SizeBefore));
+		OutTotals.Add(TEXT("SizeAfter"), FText::AsNumber(TotalEntry->SizeAfter));
+	}
+}
+
+
+TSharedPtr<SWidget> FCookerStatsPage::GetCustomFilter( TWeakPtr<IStatsViewer> InParentStatsViewer )
+{
+	return SNew(SComboButton)
+		.ContentPadding(2.0f)
+		.Visibility_Raw(this, &FCookerStatsPage::HandleFilterComboButtonVisibility)
+		.OnGetMenuContent_Raw(this, &FCookerStatsPage::HandleFilterComboButtonGetMenuContent)
+		.ButtonContent()
+		[
+			SNew(STextBlock)
+				.Text_Raw(this, &FCookerStatsPage::HandleFilterComboButtonText)
+				.ToolTipText(LOCTEXT("FilterColumnToUse_Tooltip", "Choose the target platform to filter when searching"))
+		];
+}
+
+
+void FCookerStatsPage::OnShow( TWeakPtr<class IStatsViewer> InParentStatsViewer )
+{
+
+}
+
+
+void FCookerStatsPage::OnHide( )
+{
+
+}
+
+
+/* FCookerStatsPage callback
+ *****************************************************************************/
+
+TSharedRef<SWidget> FCookerStatsPage::HandleFilterComboButtonGetMenuContent( ) const
+{
+	// visitor for the directory enumerator
+	class FCookedDirectoryVisitor
+		: public IPlatformFile::FDirectoryVisitor
+	{
+	public:
+
+		FCookedDirectoryVisitor( const FCookerStatsPage& InCookerStatsPage, FMenuBuilder& InMenuBuilder )
+			: CookerStatsPage(InCookerStatsPage)
+			, MenuBuilder(InMenuBuilder)
+		{ }
+
+		virtual bool Visit( const TCHAR* FilenameOrDirectory, bool bIsDirectory ) override
+		{
+			if (bIsDirectory)
+			{
+				FString PlatformName = FPaths::GetBaseFilename(FilenameOrDirectory);
+
+				MenuBuilder.AddMenuEntry(
+					FText::FromString(PlatformName),
+					FText::FromString(PlatformName),
+					FSlateIcon(),
+					FUIAction(
+						FExecuteAction::CreateRaw(&CookerStatsPage, &FCookerStatsPage::HandleFilterMenuEntryExecute, PlatformName),
+						FCanExecuteAction(),
+						FIsActionChecked::CreateRaw(&CookerStatsPage, &FCookerStatsPage::HandleFilterMenuEntryIsChecked, PlatformName)
+					),
+					NAME_None,
+					EUserInterfaceActionType::RadioButton
+				);
+			}
+
+			return true;
+		}
+
+	private:
+
+		const FCookerStatsPage& CookerStatsPage;
+		FMenuBuilder& MenuBuilder;
+	};
+
+
+	// create menu entries for each cooked directory
+	FMenuBuilder MenuBuilder(true, nullptr);
+	{
+		FCookedDirectoryVisitor CookedDirectoryVisitor(*this, MenuBuilder);
+		FPlatformFileManager::Get().GetPlatformFile().IterateDirectory(*(FPaths::GameSavedDir() / TEXT("Cooked")), CookedDirectoryVisitor);
+	}
+
+	return MenuBuilder.MakeWidget();
+}
+
+
+FText FCookerStatsPage::HandleFilterComboButtonText( ) const
+{
+	if (SelectedPlatformName.IsEmpty())
+	{
+		return LOCTEXT("SelectPlatformLabel", "Select platform...");
+	}
+
+	return FText::FromString(SelectedPlatformName);
+}
+
+
+EVisibility FCookerStatsPage::HandleFilterComboButtonVisibility( ) const
+{
+	return EVisibility::Visible;
+}
+
+
+void FCookerStatsPage::HandleFilterMenuEntryExecute( FString PlatformName )
+{
+	SelectedPlatformName = PlatformName;
+
+	Refresh();
+}
+
+
+bool FCookerStatsPage::HandleFilterMenuEntryIsChecked( FString PlatformName ) const
+{
+	return (PlatformName == SelectedPlatformName);
+}
+
+
+#undef LOCTEXT_NAMESPACE
diff --git a/Engine/Source/Editor/StatsViewer/Private/StatsPages/CookerStatsPage.h b/Engine/Source/Editor/StatsViewer/Private/StatsPages/CookerStatsPage.h
new file mode 100644
--- /dev/null
+++ b/Engine/Source/Editor/StatsViewer/Private/StatsPages/CookerStatsPage.h
@@ -0,0 +1,52 @@
+// Copyright 1998-2014 Epic Games, Inc. All Rights Reserved.
+
+#pragma once
+
+
+#include "CookerStats.h"
+
+/**
+ * Implements a Stats page representing cooker statistics.
+ */
+class FCookerStatsPage
+	: public FStatsPage<UCookerStats>
+{
+public:
+
+	/** Singleton accessor. */
+	static FCookerStatsPage& Get( );
+
+	/** Virtual destructor. */
+	virtual ~FCookerStatsPage( ) { }
+
+public:
+
+	// IStatsPage interface
+
+	virtual void Generate( TArray<TWeakObjectPtr<UObject>>& OutObjects ) const override;
+	virtual void GenerateTotals( const TArray<TWeakObjectPtr<UObject>>& InObjects, TMap<FString, FText>& OutTotals ) const override;
+	virtual TSharedPtr<SWidget> GetCustomFilter( TWeakPtr<class IStatsViewer> InParentStatsViewer ) override;
+	virtual void OnShow( TWeakPtr<class IStatsViewer> InParentStatsViewer ) override;
+	virtual void OnHide( ) override;
+
+private:
+
+	/** Callback for getting the filter combo button's menu content. */
+	TSharedRef<SWidget> HandleFilterComboButtonGetMenuContent( ) const;
+
+	/** Callback for getting the filter combo button's text label. */
+	FText HandleFilterComboButtonText( ) const;
+
+	/** Callback for getting the filter combo button's visibility. */
+	EVisibility HandleFilterComboButtonVisibility( ) const;
+
+	/** Callback for executing a filter menu entry. */
+	void HandleFilterMenuEntryExecute( FString PlatformName );
+
+	/** Callback for determining whether a filter menu entry should be marked checked. */
+	bool HandleFilterMenuEntryIsChecked( FString PlatformName ) const;
+
+private:
+
+	FString SelectedPlatformName;
+};
diff --git a/Engine/Source/Editor/StatsViewer/Private/StatsViewerModule.cpp b/Engine/Source/Editor/StatsViewer/Private/StatsViewerModule.cpp
--- a/Engine/Source/Editor/StatsViewer/Private/StatsViewerModule.cpp
+++ b/Engine/Source/Editor/StatsViewer/Private/StatsViewerModule.cpp
@@ -2,6 +2,7 @@
 
 #include "StatsViewerPrivatePCH.h"
 #include "StatsPageManager.h"
+#include "CookerStatsPage.h"
 #include "LightingBuildInfoStatsPage.h"
 #include "PrimitiveStatsPage.h"
 #include "StaticMeshLightingInfoStatsPage.h"
@@ -14,6 +15,8 @@
 IMPLEMENT_MODULE( FStatsViewerModule, StatsViewer );
 
 const FName StatsViewerApp = FName("StatsViewerApp");
+
+static const FName CookerStatsPage = FName("CookerStats");
 static const FName LightingBuildInfoPage = FName("LightingBuildInfo");
 static const FName PrimitiveStatsPage = FName("PrimitiveStats");
 static const FName StaticMeshLightingInfoPage = FName("StaticMeshLightingInfo");
@@ -21,6 +24,7 @@ static const FName TextureStatsPage = FName("TextureStats");
 
 void FStatsViewerModule::StartupModule()
 {
+	FStatsPageManager::Get().RegisterPage( MakeShareable(&FCookerStatsPage::Get()) );
 	FStatsPageManager::Get().RegisterPage( MakeShareable(&FLightingBuildInfoStatsPage::Get()) );
 	FStatsPageManager::Get().RegisterPage( MakeShareable(&FPrimitiveStatsPage::Get()) );
 	FStatsPageManager::Get().RegisterPage( MakeShareable(&FStaticMeshLightingInfoStatsPage::Get()) );
@@ -57,6 +61,8 @@ TSharedPtr< IStatsPage > FStatsViewerModule::GetPage( EStatsPage::Type InType )
 {
 	switch( InType )
 	{
+	case EStatsPage::CookerStats:
+		return GetPage(CookerStatsPage);
 	case EStatsPage::LightingBuildInfo:
 		return GetPage(LightingBuildInfoPage);
 	case EStatsPage::PrimitiveStats:
diff --git a/Engine/Source/Editor/StatsViewer/Public/StatsViewerModule.h b/Engine/Source/Editor/StatsViewer/Public/StatsViewerModule.h
--- a/Engine/Source/Editor/StatsViewer/Public/StatsViewerModule.h
+++ b/Engine/Source/Editor/StatsViewer/Public/StatsViewerModule.h
@@ -14,6 +14,7 @@ namespace EStatsPage
 {
 	enum Type
 	{
+		CookerStats,
 		LightingBuildInfo,
 		PrimitiveStats,
 		StaticMeshLightingInfo,
diff --git a/Engine/Source/Editor/StatsViewer/StatsViewer.Build.cs b/Engine/Source/Editor/StatsViewer/StatsViewer.Build.cs
--- a/Engine/Source/Editor/StatsViewer/StatsViewer.Build.cs
+++ b/Engine/Source/Editor/StatsViewer/StatsViewer.Build.cs
@@ -8,6 +8,7 @@ public class StatsViewer : ModuleRules
 	{
 		PrivateDependencyModuleNames.AddRange(
 			new string[] {
+				"AssetRegistry",
 				"Core",
 				"CoreUObject",
 				"Engine",
