Fixed instanced vertex color memory being misrepresented in Primitive Stats

From a report from The Coalition in this UDN: https://udn.unrealengine.com/questions/255641/accumulating-instvertexcolormem-in-bytes-vs-kb.html

diff --git a/Engine/Source/Editor/StatsViewer/Private/StatsPages/PrimitiveStatsPage.cpp b/Engine/Source/Editor/StatsViewer/Private/StatsPages/PrimitiveStatsPage.cpp
--- a/Engine/Source/Editor/StatsViewer/Private/StatsPages/PrimitiveStatsPage.cpp
+++ b/Engine/Source/Editor/StatsViewer/Private/StatsPages/PrimitiveStatsPage.cpp
@@ -242,7 +242,7 @@ struct PrimitiveStatsGenerator
 				else if( StaticMeshComponent )
 				{
 					// This stat is used by multiple components so accumulate instanced vertex color memory.
-					StatsEntry->InstVertexColorMem += InstVertexColorMem;
+					StatsEntry->InstVertexColorMem += (float)InstVertexColorMem / 1024.0f;
 				}
 				else if (LandscapeComponent)
 				{
