Always gracefully handle null world in primitive stats

Merging 2672667

//depot/UE4-UT/...

to //depot/UE4/...

diff --git a/Engine/Source/Editor/StatsViewer/Private/StatsPages/PrimitiveStatsPage.cpp b/Engine/Source/Editor/StatsViewer/Private/StatsPages/PrimitiveStatsPage.cpp
--- a/Engine/Source/Editor/StatsViewer/Private/StatsPages/PrimitiveStatsPage.cpp
+++ b/Engine/Source/Editor/StatsViewer/Private/StatsPages/PrimitiveStatsPage.cpp
@@ -168,9 +168,10 @@ struct PrimitiveStatsGenerator
 		UWorld* World = InPrimitiveComponent->GetWorld();
 	//	check(World); // @todo: re-instate this check once the GWorld migration has completed
 		/// If we should skip the actor. Skip if the actor has no outer or if we are only showing selected actors and the actor isn't selected
-		const bool bShouldSkip = World == NULL || ActorOuter == NULL || (ActorOuter != NULL && InObjectSet == PrimitiveObjectSets_SelectedObjects && ActorOuter->IsSelected() == false );
+		const bool bShouldSkip = ActorOuter == NULL || (ActorOuter != NULL && InObjectSet == PrimitiveObjectSets_SelectedObjects && ActorOuter->IsSelected() == false );
 		// Dont' care about components without a resource.
 		if(	Resource 
+			&& World != NULL
 			// Require actor association for selection and to disregard mesh emitter components. The exception being model components.
 			&&	(!bShouldSkip || (ModelComponent && InObjectSet != PrimitiveObjectSets_SelectedObjects ) )
 			// Only list primitives in visible levels
