Fixed static analysis warning in PrimitiveStatsPage

diff --git a/Engine/Source/Editor/StatsViewer/Private/StatsPages/PrimitiveStatsPage.cpp b/Engine/Source/Editor/StatsViewer/Private/StatsPages/PrimitiveStatsPage.cpp
--- a/Engine/Source/Editor/StatsViewer/Private/StatsPages/PrimitiveStatsPage.cpp
+++ b/Engine/Source/Editor/StatsViewer/Private/StatsPages/PrimitiveStatsPage.cpp
@@ -416,7 +416,7 @@ void FPrimitiveStatsPage::Generate( TArray< TWeakObjectPtr<UObject> >& OutObject
 	PrimitiveStatsGenerator Generator;
 	Generator.Generate();
 	
-	UWorld* World = (GEditor && GEditor->PlayWorld) ? GEditor->PlayWorld : GWorld;
+	UWorld* World = GEditor->PlayWorld ? GEditor->PlayWorld : GWorld;
 
 	switch ((EPrimitiveObjectSets)ObjectSetIndex)
 	{
