- Optimized Statistics Viewer : When modifying the filter, avoid rebuilding the object list (changing stats page or using Refresh button will force a full refresh)
- Fixed leak in viewer when refreshing and filtering was already applied on generated object list

diff --git a/Engine/Source/Editor/StatsViewer/Private/SStatsViewer.cpp b/Engine/Source/Editor/StatsViewer/Private/SStatsViewer.cpp
--- a/Engine/Source/Editor/StatsViewer/Private/SStatsViewer.cpp
+++ b/Engine/Source/Editor/StatsViewer/Private/SStatsViewer.cpp
@@ -225,6 +225,7 @@ void SStatsViewer::Construct( const FArguments& InArgs )
 
 SStatsViewer::SStatsViewer() :
 	bNeedsRefresh( false ),
+	bNeedsRefreshForFilterChange( false ),
 	CurrentObjectSetIndex( 0 ),
 	CurrentFilterIndex( 0 ),
 	CustomColumn( new FStatsCustomColumn )
@@ -294,7 +295,7 @@ void SStatsViewer::Tick( const FGeometry& AllottedGeometry, const double InCurre
 	SearchTextUpdateTimer -= InDeltaTime;
 	if( bTimerActive && SearchTextUpdateTimer < 0.0f )
 	{
-		bNeedsRefresh = true;
+		bNeedsRefreshForFilterChange = true;
 	}
 
 	if( CurrentStats.IsValid() )
@@ -306,47 +307,59 @@ void SStatsViewer::Tick( const FGeometry& AllottedGeometry, const double InCurre
 		}
 	}
 
-	if( bNeedsRefresh )
+	if( bNeedsRefresh || bNeedsRefreshForFilterChange)
 	{
 		if( CurrentStats.IsValid() )
 		{
-			// Flag all the current stat objects for death
-			for(auto Iter = CurrentObjects.CreateIterator(); Iter; Iter++)
+			if (bNeedsRefresh)
 			{
-				if((*Iter).IsValid())
+				// Flag all the current stat objects for death
+				for (auto Iter = LastGeneratedObjectList.CreateIterator(); Iter; Iter++)
 				{
-					(*Iter)->RemoveFromRoot();
+					if ((*Iter).IsValid())
+					{
+						(*Iter)->RemoveFromRoot();
+					}
 				}
+
+				// Generate new set of objects
+				LastGeneratedObjectList.Empty();
+				CurrentObjects.Empty();
+				CurrentStats->Generate(CurrentObjects);
+
+				// Backup list for future use (see bNeedsRefreshForFilterChange)
+				LastGeneratedObjectList = CurrentObjects;
 			}
-			CurrentObjects.Empty();
-			
+			else if (bNeedsRefreshForFilterChange)
+			{
+				// For a filter change, recycle last generated object list
+				CurrentObjects = LastGeneratedObjectList;
+			}
+
 			// clear the map of total strings
 			CustomColumn->TotalsMap.Empty();
 
-			// generate new set of objects
-			CurrentStats->Generate( CurrentObjects );
-
 			// plug objects into table
-			PropertyTable->SetObjects( CurrentObjects );
+			PropertyTable->SetObjects(CurrentObjects);
 
 			// freeze & resize columns & sort if required
 			const TArray< TSharedRef< IPropertyTableColumn > >& Columns = PropertyTable->GetColumns();
-			for( int32 ColumnIndex = 0; ColumnIndex < Columns.Num(); ++ColumnIndex )
+			for (int32 ColumnIndex = 0; ColumnIndex < Columns.Num(); ++ColumnIndex)
 			{
 				TSharedRef< IPropertyTableColumn > Column = Columns[ColumnIndex];
-				if( Columns[ColumnIndex]->GetDataSource()->IsValid() )
+				if (Columns[ColumnIndex]->GetDataSource()->IsValid())
 				{
 					TSharedPtr< FPropertyPath > PropertyPath = Column->GetDataSource()->AsPropertyPath();
 					const FPropertyInfo& PropertyInfo = PropertyPath->GetRootProperty();
 					const FString& ColumnWidthString = PropertyInfo.Property->GetMetaData(StatsViewerMetadata::ColumnWidth);
-					const float ColumnWidth = ColumnWidthString.Len() > 0 ? FCString::Atof( *ColumnWidthString ) : 100.0f;
-					Column->SetWidth( ColumnWidth );
+					const float ColumnWidth = ColumnWidthString.Len() > 0 ? FCString::Atof(*ColumnWidthString) : 100.0f;
+					Column->SetWidth(ColumnWidth);
 
 					const FString& SortModeString = PropertyInfo.Property->GetMetaData(StatsViewerMetadata::SortMode);
-					if( SortModeString.Len() > 0 )
+					if (SortModeString.Len() > 0)
 					{
-						EColumnSortMode::Type SortType = SortModeString == TEXT( "Ascending" ) ? EColumnSortMode::Ascending : EColumnSortMode::Descending;
-						PropertyTable->SortByColumn( Column, SortType, EColumnSortPriority::Primary);
+						EColumnSortMode::Type SortType = SortModeString == TEXT("Ascending") ? EColumnSortMode::Ascending : EColumnSortMode::Descending;
+						PropertyTable->SortByColumn(Column, SortType, EColumnSortPriority::Primary);
 					}
 				}
 
@@ -399,6 +412,7 @@ void SStatsViewer::Tick( const FGeometry& AllottedGeometry, const double InCurre
 		}
 
 		bNeedsRefresh = false;
+		bNeedsRefreshForFilterChange = false;
 	}
 }
 
@@ -808,7 +822,7 @@ void SStatsViewer::SetSearchFilter( int32 InFilterIndex )
 {
 	CurrentFilterIndex = InFilterIndex;
 
-	Refresh();
+	bNeedsRefreshForFilterChange = true;
 }
 
 bool SStatsViewer::IsSearchFilterSelected( int32 InFilterIndex ) const
diff --git a/Engine/Source/Editor/StatsViewer/Private/SStatsViewer.h b/Engine/Source/Editor/StatsViewer/Private/SStatsViewer.h
--- a/Engine/Source/Editor/StatsViewer/Private/SStatsViewer.h
+++ b/Engine/Source/Editor/StatsViewer/Private/SStatsViewer.h
@@ -134,6 +134,9 @@ private:
 	/** Flag to refresh the table next tick */
 	bool bNeedsRefresh;
 
+	/** Flag to refresh the table triggered by a filter change */
+	bool bNeedsRefreshForFilterChange;
+
 	/** Timer to prevent constant update of the searched items when typing */
 	float SearchTextUpdateTimer;
 
@@ -158,6 +161,9 @@ private:
 	/** The current set of objects we are viewing */
 	TArray< TWeakObjectPtr<UObject> > CurrentObjects;
 
+	/** The last generated set of objects that will be used when refresh triggered by a filter change */
+	TArray< TWeakObjectPtr<UObject> > LastGeneratedObjectList;
+
 	/** The 'total' custom column used for displaying totals for properties that support the feature */
 	TSharedRef<class FStatsCustomColumn> CustomColumn;
 
