Primitive Stats Page now supports reporting stats in PIE/SIE

diff --git a/Engine/Source/Editor/StatsViewer/Private/StatsPages/PrimitiveStatsPage.cpp b/Engine/Source/Editor/StatsViewer/Private/StatsPages/PrimitiveStatsPage.cpp
--- a/Engine/Source/Editor/StatsViewer/Private/StatsPages/PrimitiveStatsPage.cpp
+++ b/Engine/Source/Editor/StatsViewer/Private/StatsPages/PrimitiveStatsPage.cpp
@@ -415,6 +415,8 @@ void FPrimitiveStatsPage::Generate( TArray< TWeakObjectPtr<UObject> >& OutObject
 {
 	PrimitiveStatsGenerator Generator;
 	Generator.Generate();
+	
+	UWorld* World = (GEditor && GEditor->PlayWorld) ? GEditor->PlayWorld : GWorld;
 
 	switch ((EPrimitiveObjectSets)ObjectSetIndex)
 	{
@@ -424,7 +426,7 @@ void FPrimitiveStatsPage::Generate( TArray< TWeakObjectPtr<UObject> >& OutObject
 			{
 				AActor* Owner = (*It)->GetOwner();
 
-				if (Owner != nullptr && !Owner->HasAnyFlags(RF_ClassDefaultObject) && Owner->IsInLevel(GWorld->GetCurrentLevel()))
+				if (Owner != nullptr && !Owner->HasAnyFlags(RF_ClassDefaultObject) && Owner->IsInLevel(World->GetCurrentLevel()))
 				{
 					UPrimitiveStats* StatsEntry = Generator.Add(*It, (EPrimitiveObjectSets)ObjectSetIndex);
 
@@ -439,7 +441,7 @@ void FPrimitiveStatsPage::Generate( TArray< TWeakObjectPtr<UObject> >& OutObject
 
 		case PrimitiveObjectSets_AllObjects:
 		{
-			if (UWorld* World = GWorld)
+			if (World)
 			{
 				TArray<ULevel*> Levels;
 
