Fixed a crash while clearing the Stat Pages.

diff --git a/Engine/Source/Editor/StatsViewer/Private/StatsViewerModule.cpp b/Engine/Source/Editor/StatsViewer/Private/StatsViewerModule.cpp
--- a/Engine/Source/Editor/StatsViewer/Private/StatsViewerModule.cpp
+++ b/Engine/Source/Editor/StatsViewer/Private/StatsViewerModule.cpp
@@ -130,7 +130,7 @@ void FStatsViewerModule::Clear()
 {
 	for( int32 PageIndex = 0; PageIndex < FStatsPageManager::Get().NumPages(); PageIndex++ )
 	{
-		FStatsPageManager::Get().GetPage(PageIndex)->Clear();
+		FStatsPageManager::Get().GetPageByIndex(PageIndex)->Clear();
 	}
 }
 
