Fixed unhandled exception creating a new Level after opening and closing the Stats window #jira UE-82889 #rb none

diff --git a/Engine/Source/Editor/StatsViewer/Private/SStatsViewer.cpp b/Engine/Source/Editor/StatsViewer/Private/SStatsViewer.cpp
--- a/Engine/Source/Editor/StatsViewer/Private/SStatsViewer.cpp
+++ b/Engine/Source/Editor/StatsViewer/Private/SStatsViewer.cpp
@@ -50,11 +50,10 @@ namespace StatsViewerMetadata
 
 void SStatsViewer::Construct( const FArguments& InArgs )
 {
-	// We can use either provided page manager or use the default (global) one.
-	StatsPageManager = InArgs._StatsPageManager.IsValid() ? InArgs._StatsPageManager : MakeShareable(&FStatsPageManager::Get());
-
 	FPropertyEditorModule& PropertyEditorModule = FModuleManager::GetModuleChecked<FPropertyEditorModule>( "PropertyEditor" );
 
+	StatsPageManagerPtr = InArgs._StatsPageManager;
+
 	// create empty property table
 	PropertyTable = PropertyEditorModule.CreatePropertyTable();
 	PropertyTable->SetIsUserAllowedToChangeRoot( false );
@@ -64,9 +63,10 @@ void SStatsViewer::Construct( const FArguments& InArgs )
 
 	// we want to customize some columns
 	TArray< TSharedRef< IPropertyTableCustomColumn > > CustomColumns;
-	for( int32 PageIndex = 0; PageIndex < StatsPageManager->NumPages(); PageIndex++ )
+	FStatsPageManager& StatsPageManager = GetStatsPageManager();
+	for( int32 PageIndex = 0; PageIndex < StatsPageManager.NumPages(); PageIndex++ )
 	{
-		TSharedRef<IStatsPage> StatsPage = StatsPageManager->GetPageByIndex( PageIndex );
+		TSharedRef<IStatsPage> StatsPage = StatsPageManager.GetPageByIndex( PageIndex );
 		TArray< TSharedRef< IPropertyTableCustomColumn > > PagesCustomColumns;
 		StatsPage->GetCustomColumns(PagesCustomColumns);
 		if(PagesCustomColumns.Num() > 0)
@@ -212,19 +212,19 @@ void SStatsViewer::Construct( const FArguments& InArgs )
 		TSharedPtr<IStatsPage> InitialStatsPage;
 		FString DisplayedStatsPageName;
 		FString ConfigKey = TEXT("DisplayedStatsPageName");
-		if (StatsPageManager->GetName() != NAME_None)
+		if (StatsPageManager.GetName() != NAME_None)
 		{
-			ConfigKey += TEXT("_") + StatsPageManager->GetName().ToString();
+			ConfigKey += TEXT("_") + StatsPageManager.GetName().ToString();
 		}
 
 		if (GConfig->GetString(*StatsViewerConstants::ConfigSectionName, *ConfigKey, DisplayedStatsPageName, GEditorPerProjectIni))
 		{
-			InitialStatsPage = StatsPageManager->GetPage(FName(*DisplayedStatsPageName));
+			InitialStatsPage = StatsPageManager.GetPage(FName(*DisplayedStatsPageName));
 		}
 		else
 		{
 			// Default to primitive stats if no config data exists yet
-			InitialStatsPage = StatsPageManager->GetPage(EStatsPage::PrimitiveStats);
+			InitialStatsPage = StatsPageManager.GetPage(EStatsPage::PrimitiveStats);
 		}
 
 		SetDisplayedStats(InitialStatsPage.ToSharedRef());
@@ -286,9 +286,10 @@ static FString GetCellString( const TSharedPtr<IPropertyTableCell> Cell, bool bG
 void SStatsViewer::Tick( const FGeometry& AllottedGeometry, const double InCurrentTime, const float InDeltaTime )
 {
 	// check if we need to switch pages - i.e. if a page wants to be shown
-	for( int32 PageIndex = 0; PageIndex < StatsPageManager->NumPages(); PageIndex++ )
+	FStatsPageManager& StatsPageManager = GetStatsPageManager();
+	for( int32 PageIndex = 0; PageIndex < StatsPageManager.NumPages(); PageIndex++ )
 	{
-		TSharedRef<IStatsPage> StatsPage = StatsPageManager->GetPageByIndex( PageIndex );
+		TSharedRef<IStatsPage> StatsPage = StatsPageManager.GetPageByIndex( PageIndex );
 		if( StatsPage->IsShowPending() )
 		{
 			SetDisplayedStats( StatsPage );
@@ -578,6 +579,12 @@ FReply SStatsViewer::OnExportClicked()
 	return FReply::Handled();
 }
 
+FStatsPageManager& SStatsViewer::GetStatsPageManager() const
+{
+	// We can use either provided page manager or use the default (global) one.
+	return StatsPageManagerPtr.IsValid() ? *StatsPageManagerPtr : FStatsPageManager::Get();
+}
+
 int32 SStatsViewer::GetObjectSetIndex() const
 {
 	return CurrentObjectSetIndex;
@@ -612,10 +619,11 @@ FText SStatsViewer::OnGetObjectSetMenuLabel() const
 TSharedRef<SWidget> SStatsViewer::OnGetDisplayMenuContent() const
 {
 	FMenuBuilder MenuBuilder(true, NULL);
+	FStatsPageManager& StatsPageManager = GetStatsPageManager();
 
-	for( int32 PageIndex = 0; PageIndex < StatsPageManager->NumPages(); PageIndex++ )
+	for( int32 PageIndex = 0; PageIndex < StatsPageManager.NumPages(); PageIndex++ )
 	{
-		TSharedRef<IStatsPage> StatsPage = StatsPageManager->GetPageByIndex( PageIndex );
+		TSharedRef<IStatsPage> StatsPage = StatsPageManager.GetPageByIndex( PageIndex );
 		MenuBuilder.AddMenuEntry( 
 			StatsPage->GetDisplayName(), 
 			StatsPage->GetToolTip(), 
diff --git a/Engine/Source/Editor/StatsViewer/Private/SStatsViewer.h b/Engine/Source/Editor/StatsViewer/Private/SStatsViewer.h
--- a/Engine/Source/Editor/StatsViewer/Private/SStatsViewer.h
+++ b/Engine/Source/Editor/StatsViewer/Private/SStatsViewer.h
@@ -131,6 +131,9 @@ private:
 	/** We just clicked the export button */
 	FReply OnExportClicked();
 
+	/** Get a valid stats page manager */
+	FStatsPageManager& GetStatsPageManager() const;
+
 private:
 
 	/** Flag to refresh the table next tick */
@@ -175,7 +178,7 @@ private:
 	/** Container for custom filters supplied by stats pages */
 	TSharedPtr< SBorder > CustomFilter;
 
-	/** The page manager that is managing our pages */
-	TSharedPtr < FStatsPageManager > StatsPageManager;
+	/** The page manager that is managing our pages. Can be invalid, in which case the global one will be used */
+	TSharedPtr < FStatsPageManager > StatsPageManagerPtr;
 };
 
