Stats for for shaders compiled by materials
- CookByTheBook now outputs a .csv file into Saved/MaterialStats
- Statistics View now lists all the csv files from that folder.
- Load is done lazily
- Added shortcut key in Asset Browser
- MaterialPath is now serialized into shadermap(if debug info is allowed)

diff --git a/Engine/Plugins/Editor/AssetManagerEditor/Source/AssetManagerEditor/Private/AssetManagerEditorCommands.cpp b/Engine/Plugins/Editor/AssetManagerEditor/Source/AssetManagerEditor/Private/AssetManagerEditorCommands.cpp
--- a/Engine/Plugins/Editor/AssetManagerEditor/Source/AssetManagerEditor/Private/AssetManagerEditorCommands.cpp
+++ b/Engine/Plugins/Editor/AssetManagerEditor/Source/AssetManagerEditor/Private/AssetManagerEditorCommands.cpp
@@ -19,6 +19,8 @@ void FAssetManagerEditorCommands::RegisterCommands()
 {
 	UI_COMMAND(ViewReferences, "Reference Viewer...", "Launches the reference viewer showing the selected assets' references", EUserInterfaceActionType::Button, FInputChord(EModifierKey::Shift | EModifierKey::Alt, EKeys::R));
 	UI_COMMAND(ViewSizeMap, "Size Map...", "Displays an interactive map showing the approximate size of this asset and everything it references", EUserInterfaceActionType::Button, FInputChord(EModifierKey::Shift | EModifierKey::Alt, EKeys::M));
+	UI_COMMAND(ViewShaderCookStatistics, "Shader Cook Statistics...", "Show Shader CookStatistics", EUserInterfaceActionType::Button, FInputChord(EModifierKey::Shift | EModifierKey::Alt | EModifierKey::Control, EKeys::S));
+	
 	UI_COMMAND(ViewAssetAudit, "Audit Assets...", "Opens the Asset Audit UI and displays information about the selected assets", EUserInterfaceActionType::Button, FInputChord(EModifierKey::Shift | EModifierKey::Alt, EKeys::A));
 
 	UI_COMMAND(OpenSelectedInAssetEditor, "Edit...", "Opens the selected asset in the relevent editor.", EUserInterfaceActionType::Button, FInputChord());
diff --git a/Engine/Plugins/Editor/AssetManagerEditor/Source/AssetManagerEditor/Private/AssetManagerEditorCommands.h b/Engine/Plugins/Editor/AssetManagerEditor/Source/AssetManagerEditor/Private/AssetManagerEditorCommands.h
--- a/Engine/Plugins/Editor/AssetManagerEditor/Source/AssetManagerEditor/Private/AssetManagerEditorCommands.h
+++ b/Engine/Plugins/Editor/AssetManagerEditor/Source/AssetManagerEditor/Private/AssetManagerEditorCommands.h
@@ -22,6 +22,9 @@ public:
 	// Shows a size map for the selected assets
 	TSharedPtr<FUICommandInfo> ViewSizeMap;
 
+	// Shows shader cook statistics
+	TSharedPtr<FUICommandInfo> ViewShaderCookStatistics;
+
 	// Adds assets to asset audit window
 	TSharedPtr<FUICommandInfo> ViewAssetAudit;
 
diff --git a/Engine/Plugins/Editor/AssetManagerEditor/Source/AssetManagerEditor/Private/AssetManagerEditorModule.cpp b/Engine/Plugins/Editor/AssetManagerEditor/Source/AssetManagerEditor/Private/AssetManagerEditorModule.cpp
--- a/Engine/Plugins/Editor/AssetManagerEditor/Source/AssetManagerEditor/Private/AssetManagerEditorModule.cpp
+++ b/Engine/Plugins/Editor/AssetManagerEditor/Source/AssetManagerEditor/Private/AssetManagerEditorModule.cpp
@@ -62,6 +62,8 @@
 #include "UObject/ConstructorHelpers.h"
 #include "Misc/HotReloadInterface.h"
 #include "Misc/ScopedSlowTask.h"
+#include "Editor/StatsViewer/Public/IStatsViewer.h"
+#include "Editor/StatsViewer/Public/StatsViewerModule.h"
 #include "Subsystems/AssetEditorSubsystem.h"
 
 
@@ -307,6 +309,7 @@ public:
 	virtual void OpenReferenceViewerUI(const TArray<FName> SelectedPackages, const FReferenceViewerParams ReferenceViewerParams = FReferenceViewerParams()) override;
 	virtual void OpenSizeMapUI(TArray<FAssetIdentifier> SelectedIdentifiers) override;
 	virtual void OpenSizeMapUI(TArray<FName> SelectedPackages) override;	
+	virtual void OpenShaderCookStatistics(TArray<FName> SelectedIdentifiers) override;
 	virtual bool GetStringValueForCustomColumn(const FAssetData& AssetData, FName ColumnName, FString& OutValue) override;
 	virtual bool GetDisplayTextForCustomColumn(const FAssetData& AssetData, FName ColumnName, FText& OutValue) override;
 	virtual bool GetIntegerValueForCustomColumn(const FAssetData& AssetData, FName ColumnName, int64& OutValue) override;
@@ -569,6 +572,7 @@ void FAssetManagerEditorModule::ShutdownModule()
 		CommandList->UnmapAction(FAssetManagerEditorCommands::Get().ViewReferences);
 		CommandList->UnmapAction(FAssetManagerEditorCommands::Get().ViewSizeMap);
 		CommandList->UnmapAction(FAssetManagerEditorCommands::Get().ViewAssetAudit);
+		CommandList->UnmapAction(FAssetManagerEditorCommands::Get().ViewShaderCookStatistics);
 
 		TArray<FLevelEditorModule::FLevelViewportMenuExtender_SelectedActors>& LevelEditorMenuExtenderDelegates = LevelEditorModule.GetAllLevelViewportContextMenuExtenders();
 		LevelEditorMenuExtenderDelegates.RemoveAll([this](const FLevelEditorModule::FLevelViewportMenuExtender_SelectedActors& Delegate) { return Delegate.GetHandle() == LevelEditorExtenderDelegateHandle; });
@@ -725,6 +729,41 @@ void FAssetManagerEditorModule::OpenSizeMapUI(TArray<FAssetIdentifier> SelectedI
 	}
 }
 
+void FAssetManagerEditorModule::OpenShaderCookStatistics(TArray<FName> SelectedPackages)
+{
+	FString SubPath;
+	FString CommonSubPath = "";
+	if(SelectedPackages.Num())
+	{
+		//Find the common path
+		FString CommonPath = SelectedPackages[0].ToString();
+		uint32 CommonIdentical = CommonPath.Len();
+		for (FName Name : SelectedPackages)
+		{
+			FString Path = Name.ToString();
+			uint32 Identical = 0;
+			uint32 NumCharacters = FMath::Min((uint32)Path.Len(), CommonIdentical);
+			for(Identical = 0; Identical < NumCharacters; ++Identical)
+			{
+				if(CommonPath[Identical] != Path[Identical])
+				{
+					break;
+				}
+			}
+			CommonIdentical = FMath::Min(Identical, CommonIdentical);
+		}
+		CommonSubPath = CommonPath.Left(CommonIdentical);
+	}
+	static const FName LevelEditorModuleName("LevelEditor");
+	static const FName LevelEditorStatsViewerTab("LevelEditorStatsViewer");
+	FLevelEditorModule& LevelEditorModule = FModuleManager::GetModuleChecked<FLevelEditorModule>(LevelEditorModuleName);
+	TSharedPtr<FTabManager> TabManager = LevelEditorModule.GetLevelEditorTabManager();
+	TSharedRef<SDockTab> Tab = TabManager->InvokeTab(LevelEditorStatsViewerTab);
+	TSharedRef<SWidget> Content = Tab->GetContent();
+	IStatsViewer* StatsView = (IStatsViewer*)&*Content;
+	StatsView->SwitchAndFilterPage(EStatsPage::ShaderCookerStats, CommonSubPath, FString("Path"));
+}
+
 void FAssetManagerEditorModule::GetAssetDataInPaths(const TArray<FString>& Paths, TArray<FAssetData>& OutAssetData)
 {
 	// Form a filter from the paths
@@ -787,6 +826,7 @@ void FAssetManagerEditorModule::CreateAssetContextMenu(FMenuBuilder& MenuBuilder
 	MenuBuilder.AddMenuEntry(FAssetManagerEditorCommands::Get().ViewReferences);
 	MenuBuilder.AddMenuEntry(FAssetManagerEditorCommands::Get().ViewSizeMap);
 	MenuBuilder.AddMenuEntry(FAssetManagerEditorCommands::Get().ViewAssetAudit);
+	MenuBuilder.AddMenuEntry(FAssetManagerEditorCommands::Get().ViewShaderCookStatistics);
 }
 
 void FAssetManagerEditorModule::OnExtendContentBrowserCommands(TSharedRef<FUICommandList> CommandList, FOnContentBrowserGetSelection GetSelectionDelegate)
@@ -806,6 +846,14 @@ void FAssetManagerEditorModule::OnExtendContentBrowserCommands(TSharedRef<FUICom
 		})
 	);
 
+
+	CommandList->MapAction(FAssetManagerEditorCommands::Get().ViewShaderCookStatistics,
+		FExecuteAction::CreateLambda([this, GetSelectionDelegate]
+	{
+		OpenShaderCookStatistics(GetContentBrowserSelectedAssetPackages(GetSelectionDelegate));
+	})
+	);
+
 	CommandList->MapAction(FAssetManagerEditorCommands::Get().ViewAssetAudit,
 		FExecuteAction::CreateLambda([this, GetSelectionDelegate]
 		{
@@ -832,6 +880,16 @@ void FAssetManagerEditorModule::OnExtendLevelEditorCommands(TSharedRef<FUIComman
 		FCanExecuteAction::CreateRaw(this, &FAssetManagerEditorModule::AreLevelEditorPackagesSelected)
 	);
 
+	CommandList->MapAction(FAssetManagerEditorCommands::Get().ViewShaderCookStatistics,
+		FExecuteAction::CreateLambda([this]
+		{
+			OpenShaderCookStatistics(GetLevelEditorSelectedAssetPackages());
+		}),
+		FCanExecuteAction::CreateRaw(this, &FAssetManagerEditorModule::AreLevelEditorPackagesSelected)
+	);
+
+
+
 	CommandList->MapAction(FAssetManagerEditorCommands::Get().ViewAssetAudit,
 		FExecuteAction::CreateLambda([this]
 		{
@@ -891,6 +949,10 @@ TSharedRef<FExtender> FAssetManagerEditorModule::OnExtendAssetEditor(const TShar
 			FAssetManagerEditorCommands::Get().ViewSizeMap,
 			FExecuteAction::CreateRaw(this, &FAssetManagerEditorModule::OpenSizeMapUI, PackageNames));
 
+		CommandList->MapAction(
+			FAssetManagerEditorCommands::Get().ViewShaderCookStatistics,
+			FExecuteAction::CreateRaw(this, &FAssetManagerEditorModule::OpenShaderCookStatistics, PackageNames));
+
 		CommandList->MapAction(
 			FAssetManagerEditorCommands::Get().ViewAssetAudit,
 			FExecuteAction::CreateRaw(this, &FAssetManagerEditorModule::OpenAssetAuditUI, PackageNames));
diff --git a/Engine/Plugins/Editor/AssetManagerEditor/Source/AssetManagerEditor/Public/AssetManagerEditorModule.h b/Engine/Plugins/Editor/AssetManagerEditor/Source/AssetManagerEditor/Public/AssetManagerEditorModule.h
--- a/Engine/Plugins/Editor/AssetManagerEditor/Source/AssetManagerEditor/Public/AssetManagerEditorModule.h
+++ b/Engine/Plugins/Editor/AssetManagerEditor/Source/AssetManagerEditor/Public/AssetManagerEditorModule.h
@@ -196,4 +196,8 @@ public:
 	/** Spawns size map with selected packages */
 	virtual void OpenSizeMapUI(TArray<FAssetIdentifier> SelectedIdentifiers) = 0;
 	virtual void OpenSizeMapUI(TArray<FName> SelectedPackages) = 0;
+
+
+	/** Open the Shader cook stats */
+	virtual void OpenShaderCookStatistics(TArray<FName> SelectedIdentifiers) = 0;
 };
diff --git a/Engine/Source/Developer/Windows/ShaderFormatD3D/Private/D3DShaderCompiler.cpp b/Engine/Source/Developer/Windows/ShaderFormatD3D/Private/D3DShaderCompiler.cpp
--- a/Engine/Source/Developer/Windows/ShaderFormatD3D/Private/D3DShaderCompiler.cpp
+++ b/Engine/Source/Developer/Windows/ShaderFormatD3D/Private/D3DShaderCompiler.cpp
@@ -1046,6 +1046,7 @@ static bool CompileAndProcessD3DShader(FString& PreprocessedShaderSource, const
 		if (Input.bGenerateDirectCompileFile)
 		{
 			FFileHelper::SaveStringToFile(CreateShaderCompilerWorkerDirectCommandLine(Input), *(Input.DumpDebugInfoPath / TEXT("DirectCompile.txt")));
+			FFileHelper::SaveStringToFile(Input.DebugDescription, *(Input.DumpDebugInfoPath / TEXT("permutation_info.txt")));
 		}
 	}
 
diff --git a/Engine/Source/Editor/StatsViewer/Classes/ShaderCookerStats.h b/Engine/Source/Editor/StatsViewer/Classes/ShaderCookerStats.h
new file mode 100644
--- /dev/null
+++ b/Engine/Source/Editor/StatsViewer/Classes/ShaderCookerStats.h
@@ -0,0 +1,50 @@
+// Copyright 1998-2019 Epic Games, Inc. All Rights Reserved.
+
+#pragma once
+
+#include "CoreMinimal.h"
+#include "UObject/ObjectMacros.h"
+#include "UObject/Object.h"
+#include "UObject/WeakObjectPtr.h"
+#include "ShaderCookerStats.generated.h"
+
+
+UENUM()
+enum EShaderCookerStatsSets
+{
+	EShaderCookerStatsSets_Default				UMETA(DisplayName = "Default", ToolTip = "Shader Cooker Sets"),
+};
+
+/** Statistics page for shader cooker stats. */
+UCLASS(Transient, MinimalAPI, meta=( DisplayName = "Shader Cooker Stats", ObjectSetType = "EShaderCookerStatsSets" ) )
+class UShaderCookerStats : public UObject
+{
+	GENERATED_UCLASS_BODY()
+
+	/** Material name */
+	UPROPERTY(VisibleAnywhere, AssetRegistrySearchable, Category = "Stats", meta = (ColumnWidth = "200"))
+	FString Name;
+
+	/** Material platform */
+	UPROPERTY(VisibleAnywhere, AssetRegistrySearchable, Category = "Stats", meta = (ColumnWidth = "50"))
+	FString Platform;
+
+
+	/** Number of shaders compilations */
+	UPROPERTY(VisibleAnywhere, AssetRegistrySearchable, Category = "Stats", meta = (DisplayName = "Compiled", ColumnWidth = "74", ShowTotal = "true"))
+	int32 Compiled;
+
+	/** Number of shaders cooked*/
+	UPROPERTY(VisibleAnywhere, AssetRegistrySearchable, Category = "Stats", meta = (DisplayName = "Cooked", ColumnWidth = "74", ShowTotal = "true"))
+	int32 Cooked;
+
+	/** Number of permutations */
+	UPROPERTY(VisibleAnywhere, AssetRegistrySearchable, Category = "Stats", meta = (DisplayName = "Permutations", ColumnWidth = "74", ShowTotal = "true"))
+	int32 Permutations;
+
+	/** Material path */
+	UPROPERTY(VisibleAnywhere, AssetRegistrySearchable, Category = "Stats", meta = (ColumnWidth = "200"))
+	FString Path;
+
+
+};
diff --git a/Engine/Source/Editor/StatsViewer/Private/SStatsViewer.cpp b/Engine/Source/Editor/StatsViewer/Private/SStatsViewer.cpp
--- a/Engine/Source/Editor/StatsViewer/Private/SStatsViewer.cpp
+++ b/Engine/Source/Editor/StatsViewer/Private/SStatsViewer.cpp
@@ -590,6 +590,45 @@ int32 SStatsViewer::GetObjectSetIndex() const
 	return CurrentObjectSetIndex;
 }
 
+void SStatsViewer::SwitchAndFilterPage(int32 Page, FString FilterValue, FString FilterProperty)
+{
+	FStatsPageManager& Manager = FStatsPageManager::Get();
+	if(Page < Manager.NumPages())
+	{
+		TSharedRef<IStatsPage> StatsPage = FStatsPageManager::Get().GetPage(Page);
+		SetDisplayedStats(StatsPage);
+
+		int32 FilterIndex = -1;
+		if(FilterProperty.Len())
+		{
+			int32 ColumnIndex = 0;
+			for (TFieldIterator<UProperty> PropertyIter( CurrentStats->GetEntryClass(), EFieldIteratorFlags::IncludeSuper ); PropertyIter; ++PropertyIter )
+			{
+				TWeakObjectPtr< UProperty > Property = *PropertyIter;
+				if( Property->HasAnyPropertyFlags(CPF_AssetRegistrySearchable) )
+				{
+					FString FilterName = Property->GetDisplayNameText().ToString();
+					if( FilterName.Len() == 0 )
+					{
+						FilterName = UEditorEngine::GetFriendlyName(Property.Get());
+					}
+					if(FilterName.Compare(FilterProperty, ESearchCase::IgnoreCase) == 0)
+					{
+						FilterIndex = ColumnIndex;
+						break;
+					}
+					ColumnIndex++;
+				}
+			}
+		}
+		if(FilterIndex >= FilterIndex)
+		{
+			FilterTextBoxWidget->SetText(FText::FromString(FilterValue));
+			SetSearchFilter(FilterIndex);
+		}
+		Refresh();
+	}
+}
 void SStatsViewer::OnFilterTextChanged( const FText& InFilterText )
 {
 	FilterText = InFilterText.ToString();
diff --git a/Engine/Source/Editor/StatsViewer/Private/SStatsViewer.h b/Engine/Source/Editor/StatsViewer/Private/SStatsViewer.h
--- a/Engine/Source/Editor/StatsViewer/Private/SStatsViewer.h
+++ b/Engine/Source/Editor/StatsViewer/Private/SStatsViewer.h
@@ -44,6 +44,7 @@ public:
 	virtual void Refresh() override;
 	TSharedPtr< class IPropertyTable > GetPropertyTable() override;
 	int32 GetObjectSetIndex() const override;
+	void SwitchAndFilterPage(int32 Page, FString FilterValue, FString FilterProperty) override;
 	/** End IStatsViewer interface */
 
 private:
diff --git a/Engine/Source/Editor/StatsViewer/Private/StatsEntries/ShaderCookerStats.cpp b/Engine/Source/Editor/StatsViewer/Private/StatsEntries/ShaderCookerStats.cpp
new file mode 100644
--- /dev/null
+++ b/Engine/Source/Editor/StatsViewer/Private/StatsEntries/ShaderCookerStats.cpp
@@ -0,0 +1,8 @@
+// Copyright 1998-2019 Epic Games, Inc. All Rights Reserved.
+
+#include "ShaderCookerStats.h"
+
+UShaderCookerStats::UShaderCookerStats(const FObjectInitializer& ObjectInitializer) :
+	Super(ObjectInitializer)
+{
+}
diff --git a/Engine/Source/Editor/StatsViewer/Private/StatsPages/ShaderCookerStatsPage.cpp b/Engine/Source/Editor/StatsViewer/Private/StatsPages/ShaderCookerStatsPage.cpp
new file mode 100644
--- /dev/null
+++ b/Engine/Source/Editor/StatsViewer/Private/StatsPages/ShaderCookerStatsPage.cpp
@@ -0,0 +1,310 @@
+// Copyright 1998-2019 Epic Games, Inc. All Rights Reserved.
+
+#include "StatsPages/ShaderCookerStatsPage.h"
+#include "Serialization/Csv/CsvParser.h"
+#include "Framework/MultiBox/MultiBoxBuilder.h"
+#include "Widgets/Input/SComboButton.h"
+#include "HAL/PlatformFilemanager.h"
+#include "HAL/FileManager.h"
+#include "Misc/FileHelper.h"
+#include "RHIDefinitions.h"
+#include "RHIShaderFormatDefinitions.inl"
+#include "ShaderCookerStatsPage.h"
+
+#define LOCTEXT_NAMESPACE "Editor.StatsViewer.ShaderCookerStats"
+
+FShaderCookerStatsPage& FShaderCookerStatsPage::Get()
+{
+	static FShaderCookerStatsPage* Instance = NULL;
+	if( Instance == NULL )
+	{
+		Instance = new FShaderCookerStatsPage;
+	}
+	return *Instance;
+}
+
+class FShaderCookerStats
+{
+public:
+	class FShaderCookerStatsSet
+	{
+	public:
+		FString Name;
+		TArray<UShaderCookerStats*> Stats;
+		bool bInitialized;
+
+	};
+
+	/** Singleton accessor */
+	static FShaderCookerStats& Get();
+	FShaderCookerStats();
+	~FShaderCookerStats();
+
+	TArray<FString> GetStatSetNames()
+	{
+		TArray<FString> Temp;
+		for(FShaderCookerStatsSet& Set : StatSets)
+		{
+			Temp.Emplace(Set.Name);
+		}
+		return Temp;
+	}
+
+	FString GetStatSetName(int32 Index)
+	{
+		return StatSets[Index].Name;
+	}
+
+	const TArray<UShaderCookerStats*>& GetShaderCookerStats(uint32 Index)
+	{
+		FShaderCookerStatsSet Set = StatSets[Index];
+		if(!Set.bInitialized)
+		{
+			Initialize(Index);
+		}
+		return StatSets[Index].Stats;
+	}
+	uint32 NumSets()
+	{
+		return StatSets.Num();
+	}
+
+	TArray<FShaderCookerStatsSet> StatSets;
+
+	void Initialize(uint32 Index);
+};
+
+void FShaderCookerStats::Initialize(uint32 Index)
+{
+	TArray<FString> PlatformNames;
+	for (int32 Platform = 0; Platform < SP_NumPlatforms; ++Platform)
+	{
+		FString FormatName = ShaderPlatformToShaderFormatName((EShaderPlatform)Platform).ToString();
+		if (FormatName.StartsWith(TEXT("SF_")))
+		{
+			FormatName = FormatName.Mid(3);
+		}
+		PlatformNames.Add(FormatName);
+	}
+	FShaderCookerStatsSet& Set = StatSets[Index];
+	FString CSVData;
+	if (FFileHelper::LoadFileToString(CSVData, *Set.Name))
+	{
+		FCsvParser CsvParser(CSVData);
+		const FCsvParser::FRows& Rows = CsvParser.GetRows();
+		int32 RowIndex = 0;
+		int32 IndexPath = -1;
+		int32 IndexPlatform = -1;
+		int32 IndexCompiled = -1;
+		int32 IndexCooked = -1;
+		int32 IndexPermutations = -1;
+		for (const TArray<const TCHAR*>& Row : Rows)
+		{
+			if (RowIndex == 0)
+			{
+				for (int32 Column = 0; Column < Row.Num(); ++Column)
+				{
+					FString Key = Row[Column];
+					if (Key == TEXT("Path"))
+					{
+						IndexPath = Column;
+					}
+					else if (Key == TEXT("Name"))
+					{
+						if (IndexPath == -1)
+						{
+							IndexPath = Column;
+						}
+					}
+					else if (Key == TEXT("Platform"))
+					{
+						IndexPlatform = Column;
+					}
+					else if (Key == TEXT("Compiled"))
+					{
+						IndexCompiled = Column;
+					}
+					else if (Key == TEXT("Cooked"))
+					{
+						IndexCooked = Column;
+					}
+					else if (Key == TEXT("Permutations"))
+					{
+						IndexPermutations = Column;
+					}
+				}
+			}
+			else
+			{
+				FString Path = IndexPath >= 0 && IndexPath < Row.Num() ? Row[IndexPath] : TEXT("?");
+#define GET_INT(Index) (Index >= 0 && Index < Row.Num() ? FCString::Atoi(Row[Index]) : 424242)
+				int32 Platform = GET_INT(IndexPlatform);
+				int32 Compiled = GET_INT(IndexCompiled);
+				int32 Cooked = GET_INT(IndexCooked);
+				int32 Permutations = GET_INT(IndexPermutations);
+#undef GET_INT
+
+
+				UShaderCookerStats* Stat = NewObject<UShaderCookerStats>();
+
+				int32 LastSlash = -1;
+				int32 LastDot = -1;
+				FString Name = Path;
+				if (Path.FindLastChar('/', LastSlash) && Path.FindLastChar('.', LastDot))
+				{
+					Name = Path.Mid(LastSlash + 1, LastDot - LastSlash - 1);
+				}
+				Stat->Name = Name;
+				Stat->Path = Path;
+				Stat->Platform = Platform < PlatformNames.Num() ? PlatformNames[Platform] : TEXT("unknown");
+				Stat->Compiled = Compiled;
+				Stat->Cooked = Cooked;
+				Stat->Permutations = Permutations;
+				Set.Stats.Emplace(Stat);
+			}
+			RowIndex++;
+		}
+		//Set.Stats.Emplace(NewStats);
+	}
+	Set.bInitialized = true;
+}
+
+
+FShaderCookerStats& FShaderCookerStats::Get()
+{
+	static FShaderCookerStats* Instance = NULL;
+	if (Instance == NULL)
+	{
+		Instance = new FShaderCookerStats;
+	}
+	return *Instance;
+}
+
+FShaderCookerStats::FShaderCookerStats()
+{
+	TArray<FString> Files;
+	FString BasePath = FString::Printf(TEXT("%s/MaterialStats/"), *FPaths::ProjectSavedDir());
+	FPlatformFileManager::Get().GetPlatformFile().FindFiles(Files, *BasePath, TEXT("csv"));
+	for (FString Filename : Files)
+	{
+		FShaderCookerStatsSet Set;
+		Set.Name = Filename;
+		Set.bInitialized = false;
+		StatSets.Emplace(Set);
+	}
+}
+FShaderCookerStats::~FShaderCookerStats()
+{
+
+}
+
+TSharedPtr<SWidget> FShaderCookerStatsPage::GetCustomWidget(TWeakPtr<IStatsViewer> InParentStatsViewer)
+{
+	if (!CustomWidget.IsValid())
+	{
+		SAssignNew(CustomWidget, SHorizontalBox)
+			+ SHorizontalBox::Slot()
+			.AutoWidth()
+			.Padding(0.0f)
+			.HAlign(HAlign_Fill)
+			[
+				SAssignNew(PlatformComboButton, SComboButton)
+				.ContentPadding(3)
+				.OnGetMenuContent(this, &FShaderCookerStatsPage::OnGetPlatformButtonMenuContent, InParentStatsViewer)
+				.ButtonContent()
+				[
+					SNew(STextBlock)
+					.Text(this, &FShaderCookerStatsPage::OnGetPlatformMenuLabel)
+					.ToolTipText(LOCTEXT("Platform_ToolTip", "Platform"))
+				]
+			];
+	}
+
+	return CustomWidget;
+}
+
+TSharedRef<SWidget> FShaderCookerStatsPage::OnGetPlatformButtonMenuContent(TWeakPtr<IStatsViewer> InParentStatsViewer) const
+{
+	FMenuBuilder MenuBuilder(true, NULL);
+	FShaderCookerStats& Stats = FShaderCookerStats::Get();
+	for (int32 Index = 0; Index < (int32)Stats.NumSets(); ++Index)
+	{
+		FString Name = Stats.GetStatSetName(Index);
+		FText MenuText = FText::FromString(Name);
+
+		MenuBuilder.AddMenuEntry(
+			MenuText,
+			MenuText,
+			FSlateIcon(),
+			FUIAction(
+				FExecuteAction::CreateSP(const_cast<FShaderCookerStatsPage*>(this), &FShaderCookerStatsPage::OnPlatformClicked, InParentStatsViewer, Index),
+				FCanExecuteAction(),
+				FIsActionChecked::CreateSP(this, &FShaderCookerStatsPage::IsPlatformSetSelected, Index)
+			),
+			NAME_None,
+			EUserInterfaceActionType::RadioButton
+		);
+	}
+	return MenuBuilder.MakeWidget();
+}
+
+void FShaderCookerStatsPage::OnPlatformClicked(TWeakPtr<IStatsViewer> InParentStatsViewer, int32 Index)
+{
+	bool bChanged = SelectedPlatform != Index;
+	SelectedPlatform = Index;
+	if(bChanged)
+	{
+		InParentStatsViewer.Pin()->Refresh();
+	}
+}
+
+bool FShaderCookerStatsPage::IsPlatformSetSelected(int32 Index) const
+{
+	return SelectedPlatform == Index;
+}
+
+FText FShaderCookerStatsPage::OnGetPlatformMenuLabel() const
+{
+	FString ActiveSetName = FShaderCookerStats::Get().GetStatSetName(SelectedPlatform);
+	FText Text = FText::FromString(ActiveSetName);
+	return Text;
+}
+
+void FShaderCookerStatsPage::Generate( TArray< TWeakObjectPtr<UObject> >& OutObjects ) const
+{
+	FShaderCookerStats& Stats = FShaderCookerStats::Get();
+	const TArray<UShaderCookerStats*>& CookStats = Stats.GetShaderCookerStats(SelectedPlatform);
+	for(UShaderCookerStats* Stat: CookStats)
+	{
+		OutObjects.Add(Stat);
+	}
+}
+
+void FShaderCookerStatsPage::GenerateTotals( const TArray< TWeakObjectPtr<UObject> >& InObjects, TMap<FString, FText>& OutTotals ) const
+{
+	if(InObjects.Num())
+	{
+		UShaderCookerStats* TotalEntry = NewObject<UShaderCookerStats>();
+
+		for( auto It = InObjects.CreateConstIterator(); It; ++It )
+		{
+			UShaderCookerStats* StatsEntry = Cast<UShaderCookerStats>( It->Get() );
+			TotalEntry->Compiled += StatsEntry->Compiled;
+			TotalEntry->Cooked += StatsEntry->Cooked;
+			TotalEntry->Permutations += StatsEntry->Permutations;
+		}
+
+		OutTotals.Add( TEXT("Compiled"), FText::AsNumber( TotalEntry->Compiled) );
+		OutTotals.Add( TEXT("Cooked"), FText::AsNumber( TotalEntry->Cooked) );
+		OutTotals.Add( TEXT("Permutations"), FText::AsNumber( TotalEntry->Permutations) );
+	}
+}
+void FShaderCookerStatsPage::OnShow( TWeakPtr< IStatsViewer > InParentStatsViewer )
+{
+}
+
+void FShaderCookerStatsPage::OnHide()
+{
+}
+
+#undef LOCTEXT_NAMESPACE
diff --git a/Engine/Source/Editor/StatsViewer/Private/StatsPages/ShaderCookerStatsPage.h b/Engine/Source/Editor/StatsViewer/Private/StatsPages/ShaderCookerStatsPage.h
new file mode 100644
--- /dev/null
+++ b/Engine/Source/Editor/StatsViewer/Private/StatsPages/ShaderCookerStatsPage.h
@@ -0,0 +1,49 @@
+// Copyright 1998-2019 Epic Games, Inc. All Rights Reserved.
+
+#pragma once
+
+#include "CoreMinimal.h"
+#include "IStatsViewer.h"
+#include "StatsPage.h"
+#include "ShaderCookerStats.h"
+
+class SComboButton;
+class SButton;
+
+
+
+/** Stats page representing material cooking stats */
+class FShaderCookerStatsPage : public FStatsPage<UShaderCookerStats>, public TSharedFromThis<FShaderCookerStatsPage>
+{
+public:
+	/** Singleton accessor */
+	static FShaderCookerStatsPage& Get();
+
+	/** Begin IStatsPage interface */
+	virtual void Generate( TArray< TWeakObjectPtr<UObject> >& OutObjects ) const override;
+	virtual void GenerateTotals( const TArray< TWeakObjectPtr<UObject> >& InObjects, TMap<FString, FText>& OutTotals ) const override;
+	virtual TSharedPtr<SWidget> GetCustomWidget(TWeakPtr< class IStatsViewer > InParentStatsViewer) override;
+	virtual void OnShow( TWeakPtr< class IStatsViewer > InParentStatsViewer ) override;
+	virtual void OnHide() override;
+	/** End IStatsPage interface */
+
+	virtual ~FShaderCookerStatsPage() {}
+
+private:
+	/** Generate label for the displayed objects combobox */
+	FText OnGetPlatformMenuLabel() const;
+	void OnPlatformClicked(TWeakPtr<IStatsViewer> InParentStatsViewer, int32 Index);
+	bool IsPlatformSetSelected(int32 Index) const;
+	TSharedRef<SWidget> OnGetPlatformButtonMenuContent(TWeakPtr< class IStatsViewer > InParentStatsViewer) const;
+
+	/** Custom widget for this page */
+	TSharedPtr<SWidget> CustomWidget;
+
+	/** Platform combo button */
+	TSharedPtr<SComboButton> PlatformComboButton;
+	
+	/** Current Platform */
+	int32 SelectedPlatform = 0;
+
+};
+
diff --git a/Engine/Source/Editor/StatsViewer/Private/StatsViewerModule.cpp b/Engine/Source/Editor/StatsViewer/Private/StatsViewerModule.cpp
--- a/Engine/Source/Editor/StatsViewer/Private/StatsViewerModule.cpp
+++ b/Engine/Source/Editor/StatsViewer/Private/StatsViewerModule.cpp
@@ -10,6 +10,7 @@
 #include "StatsPages/PrimitiveStatsPage.h"
 #include "StatsPages/StaticMeshLightingInfoStatsPage.h"
 #include "StatsPages/TextureStatsPage.h"
+#include "StatsPages/ShaderCookerStatsPage.h"
 #include "ObjectHyperlinkColumn.h"
 
 
@@ -24,6 +25,7 @@ static const FName LightingBuildInfoPage = FName("LightingBuildInfo");
 static const FName PrimitiveStatsPage = FName("PrimitiveStats");
 static const FName StaticMeshLightingInfoPage = FName("StaticMeshLightingInfo");
 static const FName TextureStatsPage = FName("TextureStats");
+static const FName ShaderCookerStatsPage = FName("ShaderCookerStats");
 
 namespace StatsViewerModuleUtils
 {
@@ -44,6 +46,7 @@ void FStatsViewerModule::StartupModule()
 	FStatsPageManager::Get().RegisterPage( EStatsPage::PrimitiveStats, MakeShareable(&FPrimitiveStatsPage::Get()) );
 	FStatsPageManager::Get().RegisterPage( EStatsPage::StaticMeshLightingInfo, MakeShareable(&FStaticMeshLightingInfoStatsPage::Get()) );
 	FStatsPageManager::Get().RegisterPage( EStatsPage::TextureStats, MakeShareable(&FTextureStatsPage::Get()) );
+	FStatsPageManager::Get().RegisterPage( EStatsPage::ShaderCookerStats, MakeShareable(&FShaderCookerStatsPage::Get()) );
 }
 
 void FStatsViewerModule::ShutdownModule()
@@ -116,6 +119,8 @@ TSharedPtr< IStatsPage > FStatsViewerModule::GetPage( EStatsPage::Type InType )
 		return GetPage(StaticMeshLightingInfoPage);
 	case EStatsPage::TextureStats:
 		return GetPage(TextureStatsPage);
+	case EStatsPage::ShaderCookerStats:
+		return GetPage(ShaderCookerStatsPage);
 	default:
 		return NULL;
 	}
diff --git a/Engine/Source/Editor/StatsViewer/Public/IStatsViewer.h b/Engine/Source/Editor/StatsViewer/Public/IStatsViewer.h
--- a/Engine/Source/Editor/StatsViewer/Public/IStatsViewer.h
+++ b/Engine/Source/Editor/StatsViewer/Public/IStatsViewer.h
@@ -22,5 +22,7 @@ public:
 
 	/** Get the currently selected object set */
 	virtual int32 GetObjectSetIndex() const = 0;
+
+	virtual void SwitchAndFilterPage(int32 Table, FString FilterValue, FString FilterProperty) = 0;
 };
 
diff --git a/Engine/Source/Editor/StatsViewer/Public/StatsViewerModule.h b/Engine/Source/Editor/StatsViewer/Public/StatsViewerModule.h
--- a/Engine/Source/Editor/StatsViewer/Public/StatsViewerModule.h
+++ b/Engine/Source/Editor/StatsViewer/Public/StatsViewerModule.h
@@ -19,6 +19,7 @@ namespace EStatsPage
 		PrimitiveStats,
 		StaticMeshLightingInfo,
 		TextureStats,
+		ShaderCookerStats,
 	};
 }
 
diff --git a/Engine/Source/Editor/UnrealEd/Private/Commandlets/CookCommandlet.cpp b/Engine/Source/Editor/UnrealEd/Private/Commandlets/CookCommandlet.cpp
--- a/Engine/Source/Editor/UnrealEd/Private/Commandlets/CookCommandlet.cpp
+++ b/Engine/Source/Editor/UnrealEd/Private/Commandlets/CookCommandlet.cpp
@@ -580,7 +580,12 @@ int32 UCookCommandlet::Main(const FString& CmdLineParams)
 
 		TArray<FString> FilesInPath;
 				
-		CookByTheBook(Platforms, FilesInPath);
+		CookByTheBook(Platforms, FilesInPath);	
+		
+		if(GShaderCompilerStats)
+		{
+			GShaderCompilerStats->WriteStats();
+		}
 
 		// Use -LogCookStats to log the results to the command line after the cook (happens automatically on a build machine)
 		COOK_STAT(
diff --git a/Engine/Source/Runtime/Engine/Private/Materials/MaterialShader.cpp b/Engine/Source/Runtime/Engine/Private/Materials/MaterialShader.cpp
--- a/Engine/Source/Runtime/Engine/Private/Materials/MaterialShader.cpp
+++ b/Engine/Source/Runtime/Engine/Private/Materials/MaterialShader.cpp
@@ -1616,7 +1616,7 @@ void FMaterialShaderMap::Compile(
 			const bool bRecreateComponentRenderStateOnCompletion = Material->IsPersistent();
 
 			// Note: using Material->IsPersistent() to detect whether this is a preview material which should have higher priority over background compiling
-			GShaderCompilingManager->AddJobs(NewJobs, bSynchronousCompile || !Material->IsPersistent(), bRecreateComponentRenderStateOnCompletion);
+			GShaderCompilingManager->AddJobs(NewJobs, bSynchronousCompile || !Material->IsPersistent(), bRecreateComponentRenderStateOnCompletion, Material->GetBaseMaterialPathName(), GetDebugDescription());
   
 			// Compile the shaders for this shader map now if the material is not deferring and deferred compiles are not enabled globally
 			if (bSynchronousCompile)
@@ -2145,6 +2145,16 @@ void FMaterialShaderMap::GetShaderPipelineList(TArray<FShaderPipeline*>& OutShad
 	}
 }
 
+uint32 FMaterialShaderMap::GetShaderNum() const
+{
+	uint32 Count = Shaders.Num();
+	for (int32 Index = 0; Index < MeshShaderMaps.Num(); Index++)
+	{
+		Count += MeshShaderMaps[Index].GetShaderNum();
+	}
+	return Count;
+}
+
 /**
  * Registers a material shader map in the global map so it can be used by materials.
  */
@@ -2276,6 +2286,11 @@ FMaterialShaderMap::~FMaterialShaderMap()
 	check(bDeletedThroughDeferredCleanup);
 	check(!bRegistered);
 #if ALLOW_SHADERMAP_DEBUG_DATA
+	if(GShaderCompilerStats != 0)
+	{
+		FString Path = !MaterialPath.IsEmpty() ? MaterialPath : GetFriendlyName();
+		GShaderCompilerStats->RegisterCookedShaders(GetShaderNum(), Platform, Path, GetDebugDescription());
+	}
 	AllMaterialShaderMaps.RemoveSwap(this);
 #endif
 }
@@ -2367,9 +2382,11 @@ void FMaterialShaderMap::Serialize(FArchive& Ar, bool bInlineShaderResources, bo
 
 #if ALLOW_SHADERMAP_DEBUG_DATA
 	Ar << FriendlyName;
+	Ar << MaterialPath;
 #else
 	FString TempString;
 	Ar << TempString;
+	Ar << TempString;
 #endif
 
 	MaterialCompilationOutput.Serialize(Ar);
diff --git a/Engine/Source/Runtime/Engine/Private/Materials/MaterialShared.cpp b/Engine/Source/Runtime/Engine/Private/Materials/MaterialShared.cpp
--- a/Engine/Source/Runtime/Engine/Private/Materials/MaterialShared.cpp
+++ b/Engine/Source/Runtime/Engine/Private/Materials/MaterialShared.cpp
@@ -1961,6 +1961,7 @@ bool FMaterial::BeginCompileShaderMap(
 
 	SCOPE_SECONDS_COUNTER(MaterialCompileTime);
 
+	NewShaderMap->SetMaterialPath(GetBaseMaterialPathName());
 	// Generate the material shader code.
 	FMaterialCompilationOutput NewCompilationOutput;
 	FHLSLMaterialTranslator MaterialTranslator(this, NewCompilationOutput, StaticParameterSet, Platform,GetQualityLevel(), ShaderMapId.FeatureLevel, TargetPlatform);
diff --git a/Engine/Source/Runtime/Engine/Private/ShaderCompiler/ShaderCompiler.cpp b/Engine/Source/Runtime/Engine/Private/ShaderCompiler/ShaderCompiler.cpp
--- a/Engine/Source/Runtime/Engine/Private/ShaderCompiler/ShaderCompiler.cpp
+++ b/Engine/Source/Runtime/Engine/Private/ShaderCompiler/ShaderCompiler.cpp
@@ -46,6 +46,7 @@
 #include "ShaderCodeLibrary.h"
 #include "MeshMaterialShaderType.h"
 #include "ShaderParameterMetadata.h"
+#include "ProfilingDebugging/DiagnosticTable.h"
 
 #define LOCTEXT_NAMESPACE "ShaderCompiler"
 
@@ -1511,6 +1512,170 @@ int32 FShaderCompileThreadRunnable::CompilingLoop()
 	return NumActiveThreads;
 }
 
+FShaderCompilerStats* GShaderCompilerStats = NULL;
+
+void FShaderCompilerStats::WriteStats()
+{
+	{
+		FlushRenderingCommands(true);
+
+		FString FileName = FString::Printf(TEXT("%s/MaterialStats/Stats-%s.csv"), *FPaths::ProjectSavedDir(), *FDateTime::Now().ToString());
+		auto DebugWriter = IFileManager::Get().CreateDebugFileWriter(*FileName);
+		FDiagnosticTableWriterCSV StatWriter(DebugWriter);
+		const TSparseArray<ShaderCompilerStats>& PlatformStats = GetShaderCompilerStats();
+
+		StatWriter.AddColumn(TEXT("Path"));
+		StatWriter.AddColumn(TEXT("Platform"));
+		StatWriter.AddColumn(TEXT("Compiled"));
+		StatWriter.AddColumn(TEXT("Cooked"));
+		StatWriter.AddColumn(TEXT("Permutations"));
+		StatWriter.AddColumn(TEXT("CompiledDouble"));
+		StatWriter.AddColumn(TEXT("CookedDouble"));
+		StatWriter.CycleRow();
+
+		
+		for(int32 Platform = 0; Platform < PlatformStats.Num(); ++Platform)
+		{
+			if(PlatformStats.IsValidIndex(Platform))
+			{
+				const ShaderCompilerStats& Stats = PlatformStats[Platform];
+				for (const auto& Pair : Stats)
+				{
+					const FString& Path = Pair.Key;
+					const FShaderCompilerStats::FShaderStats& SingleStats = Pair.Value;
+
+					StatWriter.AddColumn(*Path);
+					StatWriter.AddColumn(TEXT("%u"), Platform);
+					StatWriter.AddColumn(TEXT("%u"), SingleStats.Compiled);
+					StatWriter.AddColumn(TEXT("%u"), SingleStats.Cooked);
+					StatWriter.AddColumn(TEXT("%u"), SingleStats.PermutationCompilations.Num());
+					StatWriter.AddColumn(TEXT("%u"), SingleStats.CompiledDouble);
+					StatWriter.AddColumn(TEXT("%u"), SingleStats.CookedDouble);
+					StatWriter.CycleRow();
+				}
+			}
+		}
+	}
+	{
+
+		FString FileName = FString::Printf(TEXT("%s/MaterialStatsDebug/StatsDebug-%s.csv"), *FPaths::ProjectSavedDir(), *FDateTime::Now().ToString());
+		auto DebugWriter = IFileManager::Get().CreateDebugFileWriter(*FileName);
+		FDiagnosticTableWriterCSV StatWriter(DebugWriter);
+		const TSparseArray<ShaderCompilerStats>& PlatformStats = GetShaderCompilerStats();
+		StatWriter.AddColumn(TEXT("Name"));
+		StatWriter.AddColumn(TEXT("Platform"));
+		StatWriter.AddColumn(TEXT("Compiles"));
+		StatWriter.AddColumn(TEXT("CompilesDouble"));
+		StatWriter.AddColumn(TEXT("Uses"));
+		StatWriter.AddColumn(TEXT("UsesDouble"));
+		StatWriter.AddColumn(TEXT("PermutationString"));
+		StatWriter.CycleRow();
+
+
+		for (int32 Platform = 0; Platform < PlatformStats.Num(); ++Platform)
+		{
+			if (PlatformStats.IsValidIndex(Platform))
+			{
+				const ShaderCompilerStats& Stats = PlatformStats[Platform];
+				for (const auto& Pair : Stats)
+				{
+					const FString& Path = Pair.Key;
+					const FShaderCompilerStats::FShaderStats& SingleStats = Pair.Value;
+					for (const FShaderCompilerStats::FShaderCompilerSinglePermutationStat& Stat : SingleStats.PermutationCompilations)
+					{
+						StatWriter.AddColumn(*Path);
+						StatWriter.AddColumn(TEXT("%u"), Platform);
+						StatWriter.AddColumn(TEXT("%u"), Stat.Compiled);
+						StatWriter.AddColumn(TEXT("%u"), Stat.CompiledDouble);
+						StatWriter.AddColumn(TEXT("%u"), Stat.Cooked);
+						StatWriter.AddColumn(TEXT("%u"), Stat.CookedDouble);
+						StatWriter.AddColumn(TEXT("%s"), *Stat.PermutationString);
+						StatWriter.CycleRow();
+					}
+				}
+
+			}
+		}
+	}
+}
+void FShaderCompilerStats::RegisterCookedShaders(uint32 NumCooked, EShaderPlatform Platform, const FString MaterialPath, FString PermutationString)
+{
+	FScopeLock Lock(&CompileStatsLock);
+	if(!CompileStats.IsValidIndex(Platform))
+	{
+		ShaderCompilerStats Stats;
+		CompileStats.Insert(Platform, Stats);
+	}
+
+	FShaderCompilerStats::FShaderStats& Stats = CompileStats[Platform].FindOrAdd(MaterialPath);
+	bool bFound = false;
+	for (FShaderCompilerSinglePermutationStat& Stat : Stats.PermutationCompilations)
+	{
+		if (PermutationString == Stat.PermutationString)
+		{
+			bFound = true;
+			if (Stat.Cooked != 0)
+			{
+				Stat.CookedDouble += NumCooked;
+				Stats.CookedDouble += NumCooked;
+			}
+			else
+			{
+				Stat.Cooked = NumCooked;
+				Stats.Cooked += NumCooked;
+			}
+		}
+	}
+	if(!bFound)
+	{
+		Stats.Cooked += NumCooked;
+	}
+	if (!bFound)
+	{
+		Stats.PermutationCompilations.Emplace(PermutationString, 0, NumCooked);
+	}
+}
+
+void FShaderCompilerStats::RegisterCompiledShaders(uint32 NumCompiled, EShaderPlatform Platform, const FString MaterialPath, FString PermutationString)
+{
+	FScopeLock Lock(&CompileStatsLock);
+	if (!CompileStats.IsValidIndex(Platform))
+	{
+		ShaderCompilerStats Stats;
+		CompileStats.Insert(Platform, Stats);
+	}
+	FShaderCompilerStats::FShaderStats& Stats = CompileStats[Platform].FindOrAdd(MaterialPath);
+
+	bool bFound = false;
+	for (FShaderCompilerSinglePermutationStat& Stat : Stats.PermutationCompilations)
+	{
+		if (PermutationString == Stat.PermutationString)
+		{
+			bFound = true;
+			if (Stat.Compiled != 0)
+			{
+				Stat.CompiledDouble += NumCompiled;
+				Stats.CompiledDouble += NumCompiled;
+			}
+			else
+			{
+				Stat.Compiled = NumCompiled;
+				Stats.Compiled += NumCompiled;
+			}
+		}
+	}
+	if(!bFound)
+	{
+		Stats.Compiled += NumCompiled;
+	}
+
+
+	if (!bFound)
+	{
+		Stats.PermutationCompilations.Emplace(PermutationString, NumCompiled, 0);
+	}
+}
+
 FShaderCompilingManager* GShaderCompilingManager = NULL;
 
 FShaderCompilingManager::FShaderCompilingManager() :
@@ -1722,13 +1887,26 @@ bool FShaderCompilingManager::GetDumpShaderDebugInfo() const
 	return GDumpShaderDebugInfo != 0;
 }
 
-void FShaderCompilingManager::AddJobs(TArray<FShaderCommonCompileJob*>& NewJobs, bool bOptimizeForLowLatency, bool bRecreateComponentRenderStateOnCompletion)
+void FShaderCompilingManager::AddJobs(TArray<FShaderCommonCompileJob*>& NewJobs, bool bOptimizeForLowLatency, bool bRecreateComponentRenderStateOnCompletion, const FString MaterialBasePath, const FString PermutationString)
 {
 	check(!FPlatformProperties::RequiresCookedData());
-
 	// Lock CompileQueueSection so we can access the input and output queues
 	FScopeLock Lock(&CompileQueueSection);
 
+	check(GShaderCompilerStats)
+
+	if(NewJobs.Num())
+	{
+		FShaderCompileJob* Job = NewJobs[0]->GetSingleShaderJob();
+		if(Job) //assume that all jobs are for the same platform
+		{
+			GShaderCompilerStats->RegisterCompiledShaders(NewJobs.Num(), Job->Input.Target.GetPlatform(), MaterialBasePath, PermutationString);
+		}
+		else
+		{
+			GShaderCompilerStats->RegisterCompiledShaders(NewJobs.Num(), SP_NumPlatforms, MaterialBasePath, PermutationString);
+		}
+	}
 	if (bOptimizeForLowLatency)
 	{
 		int32 InsertIndex = 0;
@@ -2451,7 +2629,7 @@ bool FShaderCompilingManager::HandlePotentialRetryOnError(TMap<int32, FShaderMap
 				}
 
 				// Send all the shaders from this shader map through the compiler again
-				AddJobs(Results.FinishedJobs, true, Results.bRecreateComponentRenderStateOnCompletion);
+				AddJobs(Results.FinishedJobs, true, Results.bRecreateComponentRenderStateOnCompletion, FString(""), FString(""));
 			}
 		}
 
@@ -4022,7 +4200,7 @@ void VerifyGlobalShaders(EShaderPlatform Platform, bool bLoadedFromCacheFile)
 
 	if (GlobalShaderJobs.Num() > 0)
 	{
-		GShaderCompilingManager->AddJobs(GlobalShaderJobs, true, false);
+		GShaderCompilingManager->AddJobs(GlobalShaderJobs, true, false, "Globals");
 
 		const bool bAllowAsynchronousGlobalShaderCompiling =
 			// OpenGL requires that global shader maps are compiled before attaching
diff --git a/Engine/Source/Runtime/Engine/Private/ShaderDerivedDataVersion.h b/Engine/Source/Runtime/Engine/Private/ShaderDerivedDataVersion.h
--- a/Engine/Source/Runtime/Engine/Private/ShaderDerivedDataVersion.h
+++ b/Engine/Source/Runtime/Engine/Private/ShaderDerivedDataVersion.h
@@ -11,4 +11,4 @@
 // In case of merge conflicts with DDC versions, you MUST generate a new GUID and set this new
 // guid as version
 #define GLOBALSHADERMAP_DERIVEDDATA_VER			TEXT("A85F9B4CA9264C3CBE27C6356D371896")
-#define MATERIALSHADERMAP_DERIVEDDATA_VER		TEXT("8915F134E11A46E2ACFF4214A7775BDD")
+#define MATERIALSHADERMAP_DERIVEDDATA_VER		TEXT("16B6C209A72740C49B69059E25CF3CCB")
diff --git a/Engine/Source/Runtime/Engine/Public/MaterialShared.h b/Engine/Source/Runtime/Engine/Public/MaterialShared.h
--- a/Engine/Source/Runtime/Engine/Public/MaterialShared.h
+++ b/Engine/Source/Runtime/Engine/Public/MaterialShared.h
@@ -1011,6 +1011,10 @@ public:
 	/** Builds a list of the shader pipelines in a shader map. */
 	ENGINE_API void GetShaderPipelineList(TArray<FShaderPipeline*>& OutShaderPipelines) const;
 
+
+	/** Number of Shaders in Shadermap */
+	ENGINE_API uint32 GetShaderNum() const;
+
 	/** Registers a material shader map in the global map so it can be used by materials. */
 	void Register(EShaderPlatform InShaderPlatform);
 
@@ -1073,6 +1077,7 @@ public:
 #if ALLOW_SHADERMAP_DEBUG_DATA
 			+ FriendlyName.GetAllocatedSize()
 			+ DebugDescription.GetAllocatedSize()
+			+ MaterialPath.GetAllocatedSize()
 #endif
 		;
 	}
@@ -1091,9 +1096,11 @@ public:
 #if ALLOW_SHADERMAP_DEBUG_DATA
 	const FString& GetFriendlyName() const { return FriendlyName; }
 	const FString& GetDebugDescription() const { return DebugDescription; }
+	void SetMaterialPath(FString Path){ MaterialPath = Path;}
 #else
 	const FString& GetFriendlyName() const { static FString T; return T; }
 	const FString& GetDebugDescription() const { static FString T; return T; }
+	void SetMaterialPath(FString Path) { }
 #endif
 	bool RequiresSceneColorCopy() const { return MaterialCompilationOutput.RequiresSceneColorCopy(); }
 	bool NeedsSceneTextures() const { return MaterialCompilationOutput.bNeedsSceneTextures; }
@@ -1171,6 +1178,7 @@ private:
 #if ALLOW_SHADERMAP_DEBUG_DATA
 	/** The material's user friendly name, typically the object name. */
 	FString FriendlyName;
+	FString MaterialPath;
 #endif
 
 	/** The static parameter set that this shader map was compiled with */
diff --git a/Engine/Source/Runtime/Engine/Public/ShaderCompiler.h b/Engine/Source/Runtime/Engine/Public/ShaderCompiler.h
--- a/Engine/Source/Runtime/Engine/Public/ShaderCompiler.h
+++ b/Engine/Source/Runtime/Engine/Public/ShaderCompiler.h
@@ -372,6 +372,49 @@ struct FShaderMapFinalizeResults : public FShaderMapCompileResults
 	{}
 };
 
+class FShaderCompilerStats
+{
+public:
+	struct FShaderCompilerSinglePermutationStat
+	{
+		FShaderCompilerSinglePermutationStat(FString PermutationString, uint32 Compiled, uint32 Cooked)
+			: PermutationString(PermutationString)
+			, Compiled(Compiled)
+			, Cooked(Cooked)
+			, CompiledDouble(0)
+			, CookedDouble(0)
+
+		{}
+		FString PermutationString;
+		uint32 Compiled;
+		uint32 Cooked;
+		uint32 CompiledDouble;
+		uint32 CookedDouble;
+	};
+	struct FShaderStats
+	{
+		TArray<FShaderCompilerSinglePermutationStat> PermutationCompilations;
+		uint32 Compiled = 0;
+		uint32 Cooked = 0;
+		uint32 CompiledDouble = 0;
+		uint32 CookedDouble = 0;
+
+	};
+	using ShaderCompilerStats = TMap<FString, FShaderStats>;
+
+
+	ENGINE_API void RegisterCookedShaders(uint32 NumCooked, EShaderPlatform Platform, const FString MaterialPath, FString PermutationString = FString(""));
+	ENGINE_API void RegisterCompiledShaders(uint32 NumPermutations, EShaderPlatform Platform, const FString MaterialPath, FString PermutationString = FString(""));
+	ENGINE_API const TSparseArray<ShaderCompilerStats>& GetShaderCompilerStats() { return CompileStats; }
+	ENGINE_API void WriteStats();
+
+private:
+	FCriticalSection CompileStatsLock;
+	TSparseArray<ShaderCompilerStats> CompileStats;
+};
+
+
+
 /**  
  * Manager of asynchronous and parallel shader compilation.
  * This class contains an interface to enqueue and retreive asynchronous shader jobs, and manages a FShaderCompileThreadRunnable.
@@ -551,7 +594,7 @@ public:
 	 * Adds shader jobs to be asynchronously compiled. 
 	 * FinishCompilation or ProcessAsyncResults must be used to get the results.
 	 */
-	ENGINE_API void AddJobs(TArray<FShaderCommonCompileJob*>& NewJobs, bool bOptimizeForLowLatency, bool bRecreateComponentRenderStateOnCompletion);
+	ENGINE_API void AddJobs(TArray<FShaderCommonCompileJob*>& NewJobs, bool bOptimizeForLowLatency, bool bRecreateComponentRenderStateOnCompletion, const FString MaterialBasePath, FString PermutationString = FString(""));
 
 	/**
 	* Removes all outstanding compile jobs for the passed shader maps.
@@ -595,6 +638,9 @@ public:
 /** The global shader compiling thread manager. */
 extern ENGINE_API FShaderCompilingManager* GShaderCompilingManager;
 
+/** The global shader compiling stats */
+extern ENGINE_API FShaderCompilerStats* GShaderCompilerStats;
+
 /** The shader precompilers for each platform.  These are only set during the console shader compilation while cooking or in the PrecompileShaders commandlet. */
 extern class FConsoleShaderPrecompiler* GConsoleShaderPrecompilers[SP_NumPlatforms];
 
diff --git a/Engine/Source/Runtime/Launch/Private/LaunchEngineLoop.cpp b/Engine/Source/Runtime/Launch/Private/LaunchEngineLoop.cpp
--- a/Engine/Source/Runtime/Launch/Private/LaunchEngineLoop.cpp
+++ b/Engine/Source/Runtime/Launch/Private/LaunchEngineLoop.cpp
@@ -2309,6 +2309,9 @@ int32 FEngineLoop::PreInitPreStartupScreen(const TCHAR* CmdLine)
 
 	if (bEnableShaderCompile && !FPlatformProperties::RequiresCookedData())
 	{
+		check(!GShaderCompilerStats);
+		GShaderCompilerStats = new FShaderCompilerStats();
+
 		check(!GShaderCompilingManager);
 		GShaderCompilingManager = new FShaderCompilingManager();
 
@@ -5402,6 +5405,11 @@ void FEngineLoop::AppPreExit( )
 		delete GShaderCompilingManager;
 		GShaderCompilingManager = nullptr;
 	}
+	if(GShaderCompilerStats)
+	{
+		delete GShaderCompilerStats;
+		GShaderCompilerStats = nullptr;
+	}
 #endif
 }
 
diff --git a/Engine/Source/Runtime/RenderCore/Public/Shader.h b/Engine/Source/Runtime/RenderCore/Public/Shader.h
--- a/Engine/Source/Runtime/RenderCore/Public/Shader.h
+++ b/Engine/Source/Runtime/RenderCore/Public/Shader.h
@@ -2027,6 +2027,10 @@ public:
 			}
 		}
 	}
+	uint32 GetShaderNum() const
+	{
+		return Shaders.Num();
+	}
 
 	/** Builds a list of the shader pipelines in a shader map. */
 	void GetShaderPipelineList(TArray<FShaderPipeline*>& OutShaderPipelines, FShaderPipeline::EFilter Filter) const
