Fix static analysis warning V501: There are identical sub-expressions to the left and to the right of the '>=' operator: FilterIndex >= FilterIndex
Clean up some API aspects

diff --git a/Engine/Source/Editor/StatsViewer/Private/SStatsViewer.cpp b/Engine/Source/Editor/StatsViewer/Private/SStatsViewer.cpp
--- a/Engine/Source/Editor/StatsViewer/Private/SStatsViewer.cpp
+++ b/Engine/Source/Editor/StatsViewer/Private/SStatsViewer.cpp
@@ -590,7 +590,7 @@ int32 SStatsViewer::GetObjectSetIndex() const
 	return CurrentObjectSetIndex;
 }
 
-void SStatsViewer::SwitchAndFilterPage(int32 Page, FString FilterValue, FString FilterProperty)
+void SStatsViewer::SwitchAndFilterPage(int32 Page, const FString& FilterValue, const FString& FilterProperty)
 {
 	FStatsPageManager& Manager = FStatsPageManager::Get();
 	if(Page < Manager.NumPages())
@@ -604,15 +604,15 @@ void SStatsViewer::SwitchAndFilterPage(int32 Page, FString FilterValue, FString
 			int32 ColumnIndex = 0;
 			for (TFieldIterator<UProperty> PropertyIter( CurrentStats->GetEntryClass(), EFieldIteratorFlags::IncludeSuper ); PropertyIter; ++PropertyIter )
 			{
-				TWeakObjectPtr< UProperty > Property = *PropertyIter;
+				UProperty* Property = *PropertyIter;
 				if( Property->HasAnyPropertyFlags(CPF_AssetRegistrySearchable) )
 				{
 					FString FilterName = Property->GetDisplayNameText().ToString();
 					if( FilterName.Len() == 0 )
 					{
-						FilterName = UEditorEngine::GetFriendlyName(Property.Get());
+						FilterName = UEditorEngine::GetFriendlyName(Property);
 					}
-					if(FilterName.Compare(FilterProperty, ESearchCase::IgnoreCase) == 0)
+					if(FilterName == FilterProperty)
 					{
 						FilterIndex = ColumnIndex;
 						break;
@@ -621,7 +621,7 @@ void SStatsViewer::SwitchAndFilterPage(int32 Page, FString FilterValue, FString
 				}
 			}
 		}
-		if(FilterIndex >= FilterIndex)
+		if(FilterIndex >= 0)
 		{
 			FilterTextBoxWidget->SetText(FText::FromString(FilterValue));
 			SetSearchFilter(FilterIndex);
diff --git a/Engine/Source/Editor/StatsViewer/Private/SStatsViewer.h b/Engine/Source/Editor/StatsViewer/Private/SStatsViewer.h
--- a/Engine/Source/Editor/StatsViewer/Private/SStatsViewer.h
+++ b/Engine/Source/Editor/StatsViewer/Private/SStatsViewer.h
@@ -44,7 +44,7 @@ public:
 	virtual void Refresh() override;
 	TSharedPtr< class IPropertyTable > GetPropertyTable() override;
 	int32 GetObjectSetIndex() const override;
-	void SwitchAndFilterPage(int32 Page, FString FilterValue, FString FilterProperty) override;
+	void SwitchAndFilterPage(int32 Page, FString FilterValue, const FString& FilterProperty) override;
 	/** End IStatsViewer interface */
 
 private:
diff --git a/Engine/Source/Editor/StatsViewer/Public/IStatsViewer.h b/Engine/Source/Editor/StatsViewer/Public/IStatsViewer.h
--- a/Engine/Source/Editor/StatsViewer/Public/IStatsViewer.h
+++ b/Engine/Source/Editor/StatsViewer/Public/IStatsViewer.h
@@ -23,6 +23,6 @@ public:
 	/** Get the currently selected object set */
 	virtual int32 GetObjectSetIndex() const = 0;
 
-	virtual void SwitchAndFilterPage(int32 Table, FString FilterValue, FString FilterProperty) = 0;
+	virtual void SwitchAndFilterPage(int32 Table, const FString& FilterValue, const FString& FilterProperty) = 0;
 };
 
