Fix bad signature in override
#rnx

diff --git a/Engine/Source/Editor/StatsViewer/Private/SStatsViewer.h b/Engine/Source/Editor/StatsViewer/Private/SStatsViewer.h
--- a/Engine/Source/Editor/StatsViewer/Private/SStatsViewer.h
+++ b/Engine/Source/Editor/StatsViewer/Private/SStatsViewer.h
@@ -44,7 +44,7 @@ public:
 	virtual void Refresh() override;
 	TSharedPtr< class IPropertyTable > GetPropertyTable() override;
 	int32 GetObjectSetIndex() const override;
-	void SwitchAndFilterPage(int32 Page, FString FilterValue, const FString& FilterProperty) override;
+	void SwitchAndFilterPage(int32 Page, const FString& FilterValue, const FString& FilterProperty) override;
 	/** End IStatsViewer interface */
 
 private:
