Fix incremental compile issue

diff --git a/Engine/Source/Editor/StatsViewer/Private/StatsPages/ShaderCookerStatsPage.cpp b/Engine/Source/Editor/StatsViewer/Private/StatsPages/ShaderCookerStatsPage.cpp
--- a/Engine/Source/Editor/StatsViewer/Private/StatsPages/ShaderCookerStatsPage.cpp
+++ b/Engine/Source/Editor/StatsViewer/Private/StatsPages/ShaderCookerStatsPage.cpp
@@ -10,6 +10,7 @@
 #include "RHIDefinitions.h"
 #include "RHIShaderFormatDefinitions.inl"
 #include "ShaderCookerStatsPage.h"
+#include "CoreGlobals.h"
 
 #define LOCTEXT_NAMESPACE "Editor.StatsViewer.ShaderCookerStats"
 
