add more explicit includes in an attempt to make non unity happy.

diff --git a/Engine/Source/Editor/StatsViewer/Private/StatsPages/ShaderCookerStatsPage.cpp b/Engine/Source/Editor/StatsViewer/Private/StatsPages/ShaderCookerStatsPage.cpp
--- a/Engine/Source/Editor/StatsViewer/Private/StatsPages/ShaderCookerStatsPage.cpp
+++ b/Engine/Source/Editor/StatsViewer/Private/StatsPages/ShaderCookerStatsPage.cpp
@@ -11,6 +11,9 @@
 #include "RHIShaderFormatDefinitions.inl"
 #include "ShaderCookerStatsPage.h"
 #include "CoreGlobals.h"
+#include "Misc/ConfigCacheIni.h"
+#include "Misc/App.h"
+#include "Misc/CommandLine.h"
 
 #define LOCTEXT_NAMESPACE "Editor.StatsViewer.ShaderCookerStats"
 
