Support for shadercategories, for shader cook statistics.
Add Compiletime stat
Add some tooltips

diff --git a/Engine/Source/Editor/StatsViewer/Classes/ShaderCookerStats.h b/Engine/Source/Editor/StatsViewer/Classes/ShaderCookerStats.h
--- a/Engine/Source/Editor/StatsViewer/Classes/ShaderCookerStats.h
+++ b/Engine/Source/Editor/StatsViewer/Classes/ShaderCookerStats.h
@@ -29,22 +29,30 @@ class UShaderCookerStats : public UObject
 	UPROPERTY(VisibleAnywhere, AssetRegistrySearchable, Category = "Stats", meta = (ColumnWidth = "50"))
 	FString Platform;
 
+	/** Material Category*/
+	UPROPERTY(VisibleAnywhere, AssetRegistrySearchable, Category = "Stats", meta = (ColumnWidth = "50", ToolTip = "Material Category. Edit Engine/Config/ShaderConfiguration.csv & Engine/Config/ShaderConfiguration.csv to configure"))
+	FString Category;
 
 	/** Number of shaders compilations */
-	UPROPERTY(VisibleAnywhere, AssetRegistrySearchable, Category = "Stats", meta = (DisplayName = "Compiled", ColumnWidth = "74", ShowTotal = "true"))
+	UPROPERTY(VisibleAnywhere, AssetRegistrySearchable, Category = "Stats", meta = (DisplayName = "Compiled", ToolTip = "Number of shaders that was actually compiled and not present in the DDC", ColumnWidth = "74", ShowTotal = "true"))
 	int32 Compiled;
 
 	/** Number of shaders cooked*/
-	UPROPERTY(VisibleAnywhere, AssetRegistrySearchable, Category = "Stats", meta = (DisplayName = "Cooked", ColumnWidth = "74", ShowTotal = "true"))
+	UPROPERTY(VisibleAnywhere, AssetRegistrySearchable, Category = "Stats", meta = (DisplayName = "Cooked", ToolTip = "Total number of shaders cooked", ColumnWidth = "74", ShowTotal = "true"))
 	int32 Cooked;
 
 	/** Number of permutations */
-	UPROPERTY(VisibleAnywhere, AssetRegistrySearchable, Category = "Stats", meta = (DisplayName = "Permutations", ColumnWidth = "74", ShowTotal = "true"))
+	UPROPERTY(VisibleAnywhere, AssetRegistrySearchable, Category = "Stats", meta = (DisplayName = "Permutations", ToolTip = "Number of material permutations generated by switch etc.",  ColumnWidth = "74", ShowTotal = "true"))
 	int32 Permutations;
 
+	/** Time spent compiling */
+	UPROPERTY(VisibleAnywhere, AssetRegistrySearchable, Category = "Stats", meta = (DisplayName = "Compile Time", ToolTip = "Cpu time spent compiling shaders not in DDC", ColumnWidth = "74", ShowTotal = "true"))
+	float CompileTime ;
+
 	/** Material path */
 	UPROPERTY(VisibleAnywhere, AssetRegistrySearchable, Category = "Stats", meta = (ColumnWidth = "200"))
 	FString Path;
 
 
+	int32 CategoryIndex;
 };
diff --git a/Engine/Source/Editor/StatsViewer/Private/StatsPages/ShaderCookerStatsPage.cpp b/Engine/Source/Editor/StatsViewer/Private/StatsPages/ShaderCookerStatsPage.cpp
--- a/Engine/Source/Editor/StatsViewer/Private/StatsPages/ShaderCookerStatsPage.cpp
+++ b/Engine/Source/Editor/StatsViewer/Private/StatsPages/ShaderCookerStatsPage.cpp
@@ -44,6 +44,19 @@ public:
 	FShaderCookerStats();
 	~FShaderCookerStats();
 
+
+	TPair<FString, int32> FindCategory(FString Path)
+	{
+		for(TPair<FString, int32>& Pair : StatPatterns)
+		{
+			int32 Index = Path.Find(Pair.Key, ESearchCase::IgnoreCase);
+			if(Index >= 0 && Index < 2)
+			{
+				return TPair<FString, int32>(StatCategoryNames[Pair.Value], Pair.Value);
+			}
+		}
+		return TPair<FString, int32>(FString(""), 0);
+	}
 	TArray<FString> GetStatSetNames()
 	{
 		TArray<FString> Temp;
@@ -56,7 +69,14 @@ public:
 
 	FString GetStatSetName(int32 Index)
 	{
-		return StatSets[Index].Name;
+		if(Index < StatSets.Num())
+		{
+			return StatSets[Index].Name;
+		}
+		else
+		{
+			return FString("");
+		}
 	}
 
 	const TArray<UShaderCookerStats*>& GetShaderCookerStats(uint32 Index)
@@ -73,7 +93,24 @@ public:
 		return StatSets.Num();
 	}
 
+	uint32 NumCategories()
+	{
+		return StatCategoryNames.Num();
+	}
+	FString GetCategoryName(int32 Index)
+	{
+		if(Index < StatCategoryNames.Num() && Index >= 0)
+		{
+			return StatCategoryNames[Index];
+		}
+		else
+		{
+			return FString("All");
+		}
+	}
 	TArray<FShaderCookerStatsSet> StatSets;
+	TArray<FString> StatCategoryNames;
+	TArray< TPair<FString, int32> > StatPatterns;
 
 	void Initialize(uint32 Index);
 };
@@ -105,6 +142,7 @@ void FShaderCookerStats::Initialize(uint32 Index)
 		int32 IndexCompiled = -1;
 		int32 IndexCooked = -1;
 		int32 IndexPermutations = -1;
+		int32 IndexCompileTime = -1;
 		for (const TArray<const TCHAR*>& Row : Rows)
 		{
 			if (RowIndex == 0)
@@ -139,17 +177,24 @@ void FShaderCookerStats::Initialize(uint32 Index)
 					{
 						IndexPermutations = Column;
 					}
+					else if (Key == TEXT("Compiletime"))
+					{
+						IndexCompileTime = Column;
+					}
 				}
 			}
 			else
 			{
 				FString Path = IndexPath >= 0 && IndexPath < Row.Num() ? Row[IndexPath] : TEXT("?");
 #define GET_INT(Index) (Index >= 0 && Index < Row.Num() ? FCString::Atoi(Row[Index]) : 424242)
+#define GET_FLOAT(Index) (Index >= 0 && Index < Row.Num() ? FCString::Atof(Row[Index]) : 0.f)
 				int32 Platform = GET_INT(IndexPlatform);
 				int32 Compiled = GET_INT(IndexCompiled);
 				int32 Cooked = GET_INT(IndexCooked);
 				int32 Permutations = GET_INT(IndexPermutations);
+				float CompileTime = GET_FLOAT(IndexCompileTime);
 #undef GET_INT
+#undef GET_FLOAT
 
 
 				UShaderCookerStats* Stat = NewObject<UShaderCookerStats>();
@@ -167,11 +212,14 @@ void FShaderCookerStats::Initialize(uint32 Index)
 				Stat->Compiled = Compiled;
 				Stat->Cooked = Cooked;
 				Stat->Permutations = Permutations;
+				Stat->CompileTime = CompileTime;
+				TPair<FString, int32> Category = FindCategory(Path);
+				Stat->Category = Category.Key;
+				Stat->CategoryIndex = Category.Value;
 				Set.Stats.Emplace(Stat);
 			}
 			RowIndex++;
 		}
-		//Set.Stats.Emplace(NewStats);
 	}
 	Set.bInitialized = true;
 }
@@ -212,7 +260,58 @@ FShaderCookerStats::FShaderCookerStats()
 		StatSets.Emplace(Set);
 	}
 
+	TMap<FString, int32> CategoryToIndex;
+	StatCategoryNames.Add("*All*");
+
+	auto LoadCategories =[this, &CategoryToIndex](FString Path)
+	{
+		FString CSVData;
+		if(FFileHelper::LoadFileToString(CSVData, *Path))
+		{
+			FCsvParser CsvParser(CSVData);
+			const FCsvParser::FRows& Rows = CsvParser.GetRows();
+			for (const TArray<const TCHAR*>& Row : Rows)
+			{
+				if (Row.Num() != 2)
+				{
+					continue;
+				}
+				FString Category = Row[0];
+				FString Pattern = Row[1];
+				Category = Category.TrimStart().TrimEnd().TrimQuotes();
+				Pattern = Pattern.TrimStart().TrimEnd().TrimQuotes();
+				//allow comments 
+				if (Category[0] == TCHAR('#') || Category[0] == TCHAR(';'))
+				{
+					continue;
+				}
+				if (Pattern[0] == TCHAR('\\') || Pattern[0] == TCHAR('/'))
+				{
+					Pattern = Pattern.Mid(1);
+				}
+
+				int32 Index = -1;
+				if (CategoryToIndex.Contains(Category))
+				{
+					Index = CategoryToIndex[Category];
+				}
+				else
+				{
+					Index = StatCategoryNames.Num();
+					CategoryToIndex.FindOrAdd(Category) = Index;
+					StatCategoryNames.Add(Category);
+				}
+				TPair<FString, int32> Pair(Pattern, Index);
+				StatPatterns.Add(Pair);
+			}
+		}
+	};
+
+	//load from both engine & project
+	LoadCategories(FPaths::Combine(FPaths::ProjectConfigDir(), TEXT("ShaderCategories.csv")));
+	LoadCategories(FPaths::Combine(FPaths::EngineConfigDir(), TEXT("ShaderCategories.csv")));
 }
+
 FShaderCookerStats::~FShaderCookerStats()
 {
 
@@ -293,10 +392,16 @@ FText FShaderCookerStatsPage::OnGetPlatformMenuLabel() const
 void FShaderCookerStatsPage::Generate( TArray< TWeakObjectPtr<UObject> >& OutObjects ) const
 {
 	FShaderCookerStats& Stats = FShaderCookerStats::Get();
-	const TArray<UShaderCookerStats*>& CookStats = Stats.GetShaderCookerStats(SelectedPlatform);
-	for(UShaderCookerStats* Stat: CookStats)
+	if(Stats.NumSets())
 	{
-		OutObjects.Add(Stat);
+		const TArray<UShaderCookerStats*>& CookStats = Stats.GetShaderCookerStats(SelectedPlatform);
+		for(UShaderCookerStats* Stat: CookStats)
+		{
+			if(0 == ObjectSetIndex || ObjectSetIndex == Stat->CategoryIndex)
+			{
+				OutObjects.Add(Stat);
+			}
+		}
 	}
 }
 
@@ -312,10 +417,12 @@ void FShaderCookerStatsPage::GenerateTotals( const TArray< TWeakObjectPtr<UObjec
 			TotalEntry->Compiled += StatsEntry->Compiled;
 			TotalEntry->Cooked += StatsEntry->Cooked;
 			TotalEntry->Permutations += StatsEntry->Permutations;
+			TotalEntry->CompileTime += StatsEntry->CompileTime;
 		}
 
 		OutTotals.Add( TEXT("Compiled"), FText::AsNumber( TotalEntry->Compiled) );
 		OutTotals.Add( TEXT("Cooked"), FText::AsNumber( TotalEntry->Cooked) );
+		OutTotals.Add( TEXT("CompileTime"), FText::AsNumber( TotalEntry->CompileTime ));
 		OutTotals.Add( TEXT("Permutations"), FText::AsNumber( TotalEntry->Permutations) );
 	}
 }
@@ -327,4 +434,18 @@ void FShaderCookerStatsPage::OnHide()
 {
 }
 
+int32 FShaderCookerStatsPage::GetObjectSetCount() const
+{
+	return FShaderCookerStats::Get().NumCategories();
+}
+FString FShaderCookerStatsPage::GetObjectSetName(int32 InObjectSetIndex) const
+{
+	return FShaderCookerStats::Get().GetCategoryName(InObjectSetIndex);
+}
+FString FShaderCookerStatsPage::GetObjectSetToolTip(int32 InObjectSetIndex) const 
+{
+	return GetObjectSetName(InObjectSetIndex);
+}
+
+
 #undef LOCTEXT_NAMESPACE
diff --git a/Engine/Source/Editor/StatsViewer/Private/StatsPages/ShaderCookerStatsPage.h b/Engine/Source/Editor/StatsViewer/Private/StatsPages/ShaderCookerStatsPage.h
--- a/Engine/Source/Editor/StatsViewer/Private/StatsPages/ShaderCookerStatsPage.h
+++ b/Engine/Source/Editor/StatsViewer/Private/StatsPages/ShaderCookerStatsPage.h
@@ -25,6 +25,10 @@ public:
 	virtual TSharedPtr<SWidget> GetCustomWidget(TWeakPtr< class IStatsViewer > InParentStatsViewer) override;
 	virtual void OnShow( TWeakPtr< class IStatsViewer > InParentStatsViewer ) override;
 	virtual void OnHide() override;
+	virtual int32 GetObjectSetCount() const override;;
+	virtual FString GetObjectSetName(int32 InObjectSetIndex) const override;
+	virtual FString GetObjectSetToolTip(int32 InObjectSetIndex) const override;
+
 	/** End IStatsPage interface */
 
 	virtual ~FShaderCookerStatsPage() {}
diff --git a/Engine/Source/Programs/ShaderCompileWorker/Private/ShaderCompileWorker.cpp b/Engine/Source/Programs/ShaderCompileWorker/Private/ShaderCompileWorker.cpp
--- a/Engine/Source/Programs/ShaderCompileWorker/Private/ShaderCompileWorker.cpp
+++ b/Engine/Source/Programs/ShaderCompileWorker/Private/ShaderCompileWorker.cpp
@@ -18,7 +18,7 @@
 // this is for the protocol, not the data, bump if FShaderCompilerInput or ProcessInputFromArchive changes (also search for the second one with the same name, todo: put into one header file)
 const int32 ShaderCompileWorkerInputVersion = 10;
 // this is for the protocol, not the data, bump if FShaderCompilerOutput or WriteToOutputArchive changes (also search for the second one with the same name, todo: put into one header file)
-const int32 ShaderCompileWorkerOutputVersion = 5;
+const int32 ShaderCompileWorkerOutputVersion = 6;
 // this is for the protocol, not the data, bump if FShaderCompilerOutput or WriteToOutputArchive changes (also search for the second one with the same name, todo: put into one header file)
 const int32 ShaderCompileWorkerSingleJobHeader = 'S';
 // this is for the protocol, not the data, bump if FShaderCompilerOutput or WriteToOutputArchive changes (also search for the second one with the same name, todo: put into one header file)
@@ -148,7 +148,9 @@ static void ProcessCompilationJob(const FShaderCompilerInput& Input,FShaderCompi
 	}
 
 	// Compile the shader directly through the platform dll (directly from the shader dir as the working directory)
+	double TimeStart = FPlatformTime::Seconds();
 	Compiler->CompileShader(Input.ShaderFormat, Input, Output, WorkingDirectory);
+	Output.CompileTime = FPlatformTime::Seconds() - TimeStart;
 	++GNumProcessedJobs;
 }
 
diff --git a/Engine/Source/Runtime/Engine/Private/Materials/MaterialShader.cpp b/Engine/Source/Runtime/Engine/Private/Materials/MaterialShader.cpp
--- a/Engine/Source/Runtime/Engine/Private/Materials/MaterialShader.cpp
+++ b/Engine/Source/Runtime/Engine/Private/Materials/MaterialShader.cpp
@@ -1641,6 +1641,9 @@ FShader* FMaterialShaderMap::ProcessCompilationResultsForSingleJob(FShaderCompil
 	const FShaderCompileJob& CurrentJob = *SingleJob;
 	check(CurrentJob.Id == CompilingId);
 
+#if ALLOW_SHADERMAP_DEBUG_DATA
+	CompileTime += SingleJob->Output.CompileTime;
+#endif
 	FShader* Shader = nullptr;
 	if (CurrentJob.VFType)
 	{
@@ -2277,6 +2280,7 @@ FMaterialShaderMap::FMaterialShaderMap(EShaderPlatform InPlatform) :
 	checkSlow(IsInGameThread() || IsAsyncLoading());
 #if ALLOW_SHADERMAP_DEBUG_DATA
 	AllMaterialShaderMaps.Add(this);
+	CompileTime = 0.f;
 #endif
 }
 
@@ -2289,7 +2293,7 @@ FMaterialShaderMap::~FMaterialShaderMap()
 	if(GShaderCompilerStats != 0)
 	{
 		FString Path = !MaterialPath.IsEmpty() ? MaterialPath : GetFriendlyName();
-		GShaderCompilerStats->RegisterCookedShaders(GetShaderNum(), Platform, Path, GetDebugDescription());
+		GShaderCompilerStats->RegisterCookedShaders(GetShaderNum(), CompileTime, Platform, Path, GetDebugDescription());
 	}
 	AllMaterialShaderMaps.RemoveSwap(this);
 #endif
diff --git a/Engine/Source/Runtime/Engine/Private/ShaderCompiler/ShaderCompiler.cpp b/Engine/Source/Runtime/Engine/Private/ShaderCompiler/ShaderCompiler.cpp
--- a/Engine/Source/Runtime/Engine/Private/ShaderCompiler/ShaderCompiler.cpp
+++ b/Engine/Source/Runtime/Engine/Private/ShaderCompiler/ShaderCompiler.cpp
@@ -83,7 +83,7 @@ FString GetMaterialShaderMapDDCKey()
 // this is for the protocol, not the data, bump if FShaderCompilerInput or ProcessInputFromArchive changes (also search for the second one with the same name, todo: put into one header file)
 const int32 ShaderCompileWorkerInputVersion = 10;
 // this is for the protocol, not the data, bump if FShaderCompilerOutput or WriteToOutputArchive changes (also search for the second one with the same name, todo: put into one header file)
-const int32 ShaderCompileWorkerOutputVersion = 5;
+const int32 ShaderCompileWorkerOutputVersion = 6;
 // this is for the protocol, not the data, bump if FShaderCompilerOutput or WriteToOutputArchive changes (also search for the second one with the same name, todo: put into one header file)
 const int32 ShaderCompileWorkerSingleJobHeader = 'S';
 // this is for the protocol, not the data, bump if FShaderCompilerOutput or WriteToOutputArchive changes (also search for the second one with the same name, todo: put into one header file)
@@ -1571,6 +1571,7 @@ void FShaderCompilerStats::WriteStats()
 		StatWriter.AddColumn(TEXT("Compiled"));
 		StatWriter.AddColumn(TEXT("Cooked"));
 		StatWriter.AddColumn(TEXT("Permutations"));
+		StatWriter.AddColumn(TEXT("Compiletime"));
 		StatWriter.AddColumn(TEXT("CompiledDouble"));
 		StatWriter.AddColumn(TEXT("CookedDouble"));
 		StatWriter.CycleRow();
@@ -1591,6 +1592,7 @@ void FShaderCompilerStats::WriteStats()
 					StatWriter.AddColumn(TEXT("%u"), SingleStats.Compiled);
 					StatWriter.AddColumn(TEXT("%u"), SingleStats.Cooked);
 					StatWriter.AddColumn(TEXT("%u"), SingleStats.PermutationCompilations.Num());
+					StatWriter.AddColumn(TEXT("%f"), SingleStats.CompileTime);
 					StatWriter.AddColumn(TEXT("%u"), SingleStats.CompiledDouble);
 					StatWriter.AddColumn(TEXT("%u"), SingleStats.CookedDouble);
 					StatWriter.CycleRow();
@@ -1660,7 +1662,7 @@ void FShaderCompilerStats::WriteStats()
 	}
 #endif // ALLOW_DEBUG_FILES
 }
-void FShaderCompilerStats::RegisterCookedShaders(uint32 NumCooked, EShaderPlatform Platform, const FString MaterialPath, FString PermutationString)
+void FShaderCompilerStats::RegisterCookedShaders(uint32 NumCooked, float CompileTime, EShaderPlatform Platform, const FString MaterialPath, FString PermutationString)
 {
 	FScopeLock Lock(&CompileStatsLock);
 	if(!CompileStats.IsValidIndex(Platform))
@@ -1670,6 +1672,7 @@ void FShaderCompilerStats::RegisterCookedShaders(uint32 NumCooked, EShaderPlatfo
 	}
 
 	FShaderCompilerStats::FShaderStats& Stats = CompileStats[Platform].FindOrAdd(MaterialPath);
+	Stats.CompileTime += CompileTime;
 	bool bFound = false;
 	for (FShaderCompilerSinglePermutationStat& Stat : Stats.PermutationCompilations)
 	{
diff --git a/Engine/Source/Runtime/Engine/Public/ShaderCompiler.h b/Engine/Source/Runtime/Engine/Public/ShaderCompiler.h
--- a/Engine/Source/Runtime/Engine/Public/ShaderCompiler.h
+++ b/Engine/Source/Runtime/Engine/Public/ShaderCompiler.h
@@ -398,12 +398,13 @@ public:
 		uint32 Cooked = 0;
 		uint32 CompiledDouble = 0;
 		uint32 CookedDouble = 0;
+		float CompileTime = 0.f;
 
 	};
 	using ShaderCompilerStats = TMap<FString, FShaderStats>;
 
 
-	ENGINE_API void RegisterCookedShaders(uint32 NumCooked, EShaderPlatform Platform, const FString MaterialPath, FString PermutationString = FString(""));
+	ENGINE_API void RegisterCookedShaders(uint32 NumCooked, float CompileTime, EShaderPlatform Platform, const FString MaterialPath, FString PermutationString = FString(""));
 	ENGINE_API void RegisterCompiledShaders(uint32 NumPermutations, EShaderPlatform Platform, const FString MaterialPath, FString PermutationString = FString(""));
 	ENGINE_API const TSparseArray<ShaderCompilerStats>& GetShaderCompilerStats() { return CompileStats; }
 	ENGINE_API void WriteStats();
diff --git a/Engine/Source/Runtime/RenderCore/Public/ShaderCore.h b/Engine/Source/Runtime/RenderCore/Public/ShaderCore.h
--- a/Engine/Source/Runtime/RenderCore/Public/ShaderCore.h
+++ b/Engine/Source/Runtime/RenderCore/Public/ShaderCore.h
@@ -1109,6 +1109,7 @@ struct FShaderCompilerOutput
 	FShaderCompilerOutput()
 	:	NumInstructions(0)
 	,	NumTextureSamplers(0)
+	,	CompileTime(0.0)
 	,	bSucceeded(false)
 	,	bFailedRemovingUnused(false)
 	,	bSupportsQueryingUsedAttributes(false)
@@ -1123,6 +1124,7 @@ struct FShaderCompilerOutput
 	FSHAHash OutputHash;
 	uint32 NumInstructions;
 	uint32 NumTextureSamplers;
+	double CompileTime;
 	bool bSucceeded;
 	bool bFailedRemovingUnused;
 	bool bSupportsQueryingUsedAttributes;
@@ -1140,6 +1142,7 @@ struct FShaderCompilerOutput
 		// Note: this serialize is used to pass between UE4 and the shader compile worker, recompile both when modifying
 		Ar << Output.ParameterMap << Output.Errors << Output.Target << Output.ShaderCode << Output.NumInstructions << Output.NumTextureSamplers << Output.bSucceeded;
 		Ar << Output.bFailedRemovingUnused << Output.bSupportsQueryingUsedAttributes << Output.UsedAttributes;
+		Ar << Output.CompileTime;
 		Ar << Output.OptionalFinalShaderSource;
 		Ar << Output.PlatformDebugData;
 
