Fix two crashes in ShaderCookerStats when no stats are yet available

diff --git a/Engine/Source/Editor/StatsViewer/Private/StatsPages/ShaderCookerStatsPage.cpp b/Engine/Source/Editor/StatsViewer/Private/StatsPages/ShaderCookerStatsPage.cpp
--- a/Engine/Source/Editor/StatsViewer/Private/StatsPages/ShaderCookerStatsPage.cpp
+++ b/Engine/Source/Editor/StatsViewer/Private/StatsPages/ShaderCookerStatsPage.cpp
@@ -69,13 +69,13 @@ public:
 
 	FString GetStatSetName(int32 Index)
 	{
-		if(Index < StatSets.Num())
+		if((uint32)Index < NumSets())
 		{
 			return StatSets[Index].Name;
 		}
 		else
 		{
-			return FString("");
+			return "";
 		}
 	}
 
@@ -392,7 +392,7 @@ FText FShaderCookerStatsPage::OnGetPlatformMenuLabel() const
 void FShaderCookerStatsPage::Generate( TArray< TWeakObjectPtr<UObject> >& OutObjects ) const
 {
 	FShaderCookerStats& Stats = FShaderCookerStats::Get();
-	if(Stats.NumSets())
+	if((uint32)SelectedPlatform < Stats.NumSets())
 	{
 		const TArray<UShaderCookerStats*>& CookStats = Stats.GetShaderCookerStats(SelectedPlatform);
 		for(UShaderCookerStats* Stat: CookStats)
