Fix PrimitiveStatsGenerator::IsInVisibleLevel to support external actors.

diff --git a/Engine/Source/Editor/StatsViewer/Private/StatsPages/PrimitiveStatsPage.cpp b/Engine/Source/Editor/StatsViewer/Private/StatsPages/PrimitiveStatsPage.cpp
--- a/Engine/Source/Editor/StatsViewer/Private/StatsPages/PrimitiveStatsPage.cpp
+++ b/Engine/Source/Editor/StatsViewer/Private/StatsPages/PrimitiveStatsPage.cpp
@@ -45,12 +45,12 @@ struct PrimitiveStatsGenerator
 		check( InObject );
 		check( InWorld );
 
-		UObject* ObjectPackage = InObject->GetOutermost();
+		UObject* ObjectTopMostPackage = InObject->GetOutermostObject()->GetPackage();
 
 		for( int32 LevelIndex=0; LevelIndex<InWorld->GetNumLevels(); LevelIndex++ )
 		{
 			ULevel* Level = InWorld->GetLevel(LevelIndex);
-			if( Level && Level->GetOutermost() == ObjectPackage )
+			if( Level && Level->GetPackage() == ObjectTopMostPackage )
 			{
 				return true;
 			}
