Fixing shader cooker stats showing the wrong shader format names. PlatformNames is indexed by SP_ enum so we have to put each format at the correct index. Also uses the DDSPI to determine if a format is valid to display at all.

diff --git a/Engine/Source/Editor/StatsViewer/Private/StatsPages/ShaderCookerStatsPage.cpp b/Engine/Source/Editor/StatsViewer/Private/StatsPages/ShaderCookerStatsPage.cpp
--- a/Engine/Source/Editor/StatsViewer/Private/StatsPages/ShaderCookerStatsPage.cpp
+++ b/Engine/Source/Editor/StatsViewer/Private/StatsPages/ShaderCookerStatsPage.cpp
@@ -118,28 +118,28 @@ public:
 void FShaderCookerStats::Initialize(uint32 Index)
 {
 	TArray<FString> PlatformNames;
-	for (int32 Platform = 0; Platform < SP_NumPlatforms; ++Platform)
+	PlatformNames.Reserve(SP_NumPlatforms);
+
+	for (int32 PlatformIndex = 0; PlatformIndex < SP_NumPlatforms; ++PlatformIndex)
 	{
-		// ShaderPlatformToShaderFormatName asserts if it's passed a deprecated value, so we'll filter out the removed platforms here.
-		PRAGMA_DISABLE_DEPRECATION_WARNINGS
-		switch (Platform)
-		{
-			case SP_METAL_SM5_NOTESS_REMOVED:
-			case SP_VULKAN_SM5_LUMIN_REMOVED:
-			case SP_VULKAN_ES3_1_LUMIN_REMOVED:
-				continue;
-		}
-		PRAGMA_ENABLE_DEPRECATION_WARNINGS
+		const EShaderPlatform Platform = static_cast<EShaderPlatform>(PlatformIndex);
+
+		const FName ShaderFormatName = FDataDrivenShaderPlatformInfo::IsValid(Platform)
+			? ShaderPlatformToShaderFormatName(Platform) : NAME_None;
 
-		FString FormatName = ShaderPlatformToShaderFormatName((EShaderPlatform)Platform).ToString();
-		if (FormatName.Len() > 0)
+		if (ShaderFormatName != NAME_None)
 		{
+			FString FormatName = ShaderFormatName.ToString();
 			if (FormatName.StartsWith(TEXT("SF_")))
 			{
 				FormatName.MidInline(3, MAX_int32, false);
 			}
 			PlatformNames.Add(MoveTemp(FormatName));
 		}
+		else
+		{
+			PlatformNames.Add(TEXT("unknown"));
+		}
 	}
 	FShaderCookerStatsSet& Set = StatSets[Index];
 	FString CSVData;
