Fix for crash when auditing static meshes with no static lighting enabled.

diff --git a/Engine/Source/Editor/StatsViewer/Private/StatsPages/StaticMeshLightingInfoStatsPage.cpp b/Engine/Source/Editor/StatsViewer/Private/StatsPages/StaticMeshLightingInfoStatsPage.cpp
--- a/Engine/Source/Editor/StatsViewer/Private/StatsPages/StaticMeshLightingInfoStatsPage.cpp
+++ b/Engine/Source/Editor/StatsViewer/Private/StatsPages/StaticMeshLightingInfoStatsPage.cpp
@@ -133,15 +133,18 @@ struct StaticMeshLightingInfoStatsGenerator
 			Entry->StaticMesh = InComponent->GetStaticMesh();
 			
 			// Show all of the lightmap texture names so we can correlate meshes with the actual texture data
-			// for debugging encoding concerns.
-			const FMeshMapBuildData* MeshMapBuildData = InComponent->GetMeshMapBuildData(InComponent->LODData[0]);
-			if (MeshMapBuildData)
+			// for debugging encoding concerns. Only do this if we have static lighting enabled/built (LODData exists)
+			if (InComponent->LODData.Num())
 			{
-				TArray<UTexture2D*> Textures;
-				MeshMapBuildData->LightMap->GetLightMap2D()->GetReferencedTextures(Textures);
-				for (UTexture2D* Texture : Textures)
+				const FMeshMapBuildData* MeshMapBuildData = InComponent->GetMeshMapBuildData(InComponent->LODData[0]);
+				if (MeshMapBuildData)
 				{
-					Entry->LightmapTextureNames.Add(Texture->GetName());
+					TArray<UTexture2D*> Textures;
+					MeshMapBuildData->LightMap->GetLightMap2D()->GetReferencedTextures(Textures);
+					for (UTexture2D* Texture : Textures)
+					{
+						Entry->LightmapTextureNames.Add(Texture->GetName());
+					}
 				}
 			}
 
