Fix for a different but very similar crash with audition static mesh lighting. - earlier change was for disabled static lighting, this one for meshes that don't have an assigned lightmap.

diff --git a/Engine/Source/Editor/StatsViewer/Private/StatsPages/StaticMeshLightingInfoStatsPage.cpp b/Engine/Source/Editor/StatsViewer/Private/StatsPages/StaticMeshLightingInfoStatsPage.cpp
--- a/Engine/Source/Editor/StatsViewer/Private/StatsPages/StaticMeshLightingInfoStatsPage.cpp
+++ b/Engine/Source/Editor/StatsViewer/Private/StatsPages/StaticMeshLightingInfoStatsPage.cpp
@@ -137,10 +137,14 @@ struct StaticMeshLightingInfoStatsGenerator
 			if (InComponent->LODData.Num())
 			{
 				const FMeshMapBuildData* MeshMapBuildData = InComponent->GetMeshMapBuildData(InComponent->LODData[0]);
-				if (MeshMapBuildData)
+				if (MeshMapBuildData && MeshMapBuildData->LightMap)
 				{
 					TArray<UTexture2D*> Textures;
-					MeshMapBuildData->LightMap->GetLightMap2D()->GetReferencedTextures(Textures);
+					FLightMap2D* LightMap2D = MeshMapBuildData->LightMap->GetLightMap2D();
+					if (LightMap2D)
+					{
+						LightMap2D->GetReferencedTextures(Textures);
+					}
 					for (UTexture2D* Texture : Textures)
 					{
 						Entry->LightmapTextureNames.Add(Texture->GetName());
