Fix texture stats memory being reported after async texture compilation is finished

diff --git a/Engine/Source/Editor/StatsViewer/Private/StatsPages/TextureStatsPage.cpp b/Engine/Source/Editor/StatsViewer/Private/StatsPages/TextureStatsPage.cpp
--- a/Engine/Source/Editor/StatsViewer/Private/StatsPages/TextureStatsPage.cpp
+++ b/Engine/Source/Editor/StatsViewer/Private/StatsPages/TextureStatsPage.cpp
@@ -10,6 +10,7 @@
 #include "Engine/Texture2D.h"
 #include "Engine/Selection.h"
 #include "Engine/TextureCube.h"
+#include "TextureCompiler.h"
 #include "EngineUtils.h"
 #include "Editor.h"
 #include "ReferencedAssetsUtils.h"
@@ -212,8 +213,23 @@ struct TextureStatsGenerator : public FFindReferencedAssets
 			Entry->Path = GetTexturePath(InTexture->GetPathName());
 			Entry->Group = (TextureGroup)InTexture->LODGroup;
 
-			Entry->CurrentKB = InTexture->CalcTextureMemorySizeEnum( TMC_ResidentMips ) / 1024.0f;
-			Entry->FullyLoadedKB = InTexture->CalcTextureMemorySizeEnum( TMC_AllMipsBiased ) / 1024.0f;
+			// Avoid pulling on the platform data while compilation is pending
+			// as it would stall until the compilation is finished. We will report
+			// pending compiling instead for improved UI responsiveness during async
+			// compilation.
+			const bool bIsCompiling = InTexture->IsCompiling();
+
+			if (bIsCompiling)
+			{
+				// Make it abondantly clear that the value is wrong until the compilation is finished.
+				Entry->CurrentKB = -1.0f; 
+				Entry->FullyLoadedKB = -1.0f;
+			}
+			else
+			{
+				Entry->CurrentKB = InTexture->CalcTextureMemorySizeEnum( TMC_ResidentMips ) / 1024.0f;
+				Entry->FullyLoadedKB = InTexture->CalcTextureMemorySizeEnum( TMC_AllMipsBiased ) / 1024.0f;
+			}
 
 			Entry->LODBias = InTexture->GetCachedLODBias();
 
@@ -223,7 +239,7 @@ struct TextureStatsGenerator : public FFindReferencedAssets
 				Entry->LastTimeRendered = (float)FMath::Max( FApp::GetLastTime() - Resource->LastRenderTime, 0.0 );
 			}
 
-			Entry->Virtual = InTexture->IsCurrentlyVirtualTextured() ? TEXT("YES") : TEXT("NO");
+			Entry->Virtual = bIsCompiling ? TEXT("Unknown") : (InTexture->IsCurrentlyVirtualTextured() ? TEXT("YES") : TEXT("NO"));
 
 			const UTexture2D* Texture2D = Cast<const UTexture2D>(InTexture);
 			if( Texture2D )
@@ -271,6 +287,12 @@ struct TextureStatsGenerator : public FFindReferencedAssets
 #endif
 				}
 			}
+
+			if (bIsCompiling)
+			{
+				// Replace the texture type by something that convey meaning to the user that compilation is pending.
+				Entry->Type = TEXT("Compiling...");
+			}
 		}
 		else
 		{
@@ -357,8 +379,16 @@ void FTextureStatsPage::GenerateTotals( const TArray< TWeakObjectPtr<UObject> >&
 		for( auto It = InObjects.CreateConstIterator(); It; ++It )
 		{
 			UTextureStats* StatsEntry = Cast<UTextureStats>( It->Get() );
-			TotalEntry->CurrentKB += StatsEntry->CurrentKB;
-			TotalEntry->FullyLoadedKB += StatsEntry->FullyLoadedKB;
+			if (StatsEntry->CurrentKB >= 0)
+			{
+				TotalEntry->CurrentKB += StatsEntry->CurrentKB;
+			}
+			
+			if (StatsEntry->FullyLoadedKB >= 0)
+			{
+				TotalEntry->FullyLoadedKB += StatsEntry->FullyLoadedKB;
+			}
+
 			TotalEntry->NumUses += StatsEntry->NumUses;
 		}
 
@@ -392,16 +422,43 @@ void FTextureStatsPage::OnEditorNewCurrentLevel( TWeakPtr< IStatsViewer > InPare
 	}
 }
 
+void FTextureStatsPage::OnAssetPostCompile( const TArray<FAssetCompileData>& CompiledAssets, TWeakPtr< IStatsViewer > InParentStatsViewer )
+{
+	// Only trigger a refresh on the last compiled texture because rebuilding all the stats is costly enough 
+	// that we don't want to trigger a rebuild for every single texture that finishes compiling.
+	if (FTextureCompilingManager::Get().GetNumRemainingTextures() == 0)
+	{
+		if (InParentStatsViewer.IsValid())
+		{
+			// Make sure that the event is not sent for some other asset type since we don't want to
+			// trigger refresh unless it concerns textures.
+			const bool bContainsTextures =
+				CompiledAssets.ContainsByPredicate(
+					[](const FAssetCompileData& AssetCompileData)
+					{
+						return Cast<UTexture>(AssetCompileData.Asset.Get()) != nullptr;
+					});
+
+			if (bContainsTextures)
+			{
+				InParentStatsViewer.Pin()->Refresh();
+			}
+		}
+	}
+}
+
 void FTextureStatsPage::OnShow( TWeakPtr< IStatsViewer > InParentStatsViewer )
 {
 	// register delegates for scene changes we are interested in
 	USelection::SelectionChangedEvent.AddRaw(this, &FTextureStatsPage::OnEditorSelectionChanged, InParentStatsViewer);
 	FEditorDelegates::NewCurrentLevel.AddRaw(this, &FTextureStatsPage::OnEditorNewCurrentLevel, InParentStatsViewer);
+	FAssetCompilingManager::Get().OnAssetPostCompileEvent().AddRaw(this, &FTextureStatsPage::OnAssetPostCompile, InParentStatsViewer);
 }
 
 void FTextureStatsPage::OnHide()
 {
 	// unregister delegates
+	FAssetCompilingManager::Get().OnAssetPostCompileEvent().RemoveAll(this);
 	USelection::SelectionChangedEvent.RemoveAll(this);
 	FEditorDelegates::NewCurrentLevel.RemoveAll(this);
 }
diff --git a/Engine/Source/Editor/StatsViewer/Private/StatsPages/TextureStatsPage.h b/Engine/Source/Editor/StatsViewer/Private/StatsPages/TextureStatsPage.h
--- a/Engine/Source/Editor/StatsViewer/Private/StatsPages/TextureStatsPage.h
+++ b/Engine/Source/Editor/StatsViewer/Private/StatsPages/TextureStatsPage.h
@@ -6,6 +6,7 @@
 #include "IStatsViewer.h"
 #include "StatsPage.h"
 #include "TextureStats.h"
+#include "AssetCompilingManager.h"
 
 /** Stats page representing texture stats */
 class FTextureStatsPage : public FStatsPage<UTextureStats>
@@ -25,10 +26,13 @@ public:
 
 private:
 
-	/** Delegate to allow is to trigger a refresh on actor selection */
+	/** Delegate to trigger a refresh on actor selection */
 	void OnEditorSelectionChanged( UObject* NewSelection, TWeakPtr< class IStatsViewer > InParentStatsViewer );
 
-	/** Delegate to allow is to trigger a refresh on new level */
+	/** Delegate to trigger a refresh on new level */
 	void OnEditorNewCurrentLevel( TWeakPtr< class IStatsViewer > InParentStatsViewer );
+
+	/** Delegate to trigger a refresh when a texture finishes compiling */
+	void OnAssetPostCompile( const TArray<FAssetCompileData>& CompiledAssets, TWeakPtr< IStatsViewer > InParentStatsViewer );
 };
 
diff --git a/Engine/Source/Runtime/Engine/Private/TextureCompiler.cpp b/Engine/Source/Runtime/Engine/Private/TextureCompiler.cpp
--- a/Engine/Source/Runtime/Engine/Private/TextureCompiler.cpp
+++ b/Engine/Source/Runtime/Engine/Private/TextureCompiler.cpp
@@ -207,6 +207,9 @@ void FTextureCompilingManager::PostCompilation(UTexture* Texture)
 	Texture->FinishCachePlatformData();
 	Texture->UpdateResource();
 
+	// Needs to be called after the placeholder resource has been replaced by the real one
+	Texture->UpdateCachedLODBias();
+
 	// Generate an empty property changed event, to force the asset registry tag
 	// to be refreshed now that pixel format and alpha channels are available.
 	FPropertyChangedEvent EmptyPropertyChangedEvent(nullptr);
