Fix Shader cook statistics

diff --git a/Engine/Source/Editor/StatsViewer/Private/StatsPages/ShaderCookerStatsPage.cpp b/Engine/Source/Editor/StatsViewer/Private/StatsPages/ShaderCookerStatsPage.cpp
--- a/Engine/Source/Editor/StatsViewer/Private/StatsPages/ShaderCookerStatsPage.cpp
+++ b/Engine/Source/Editor/StatsViewer/Private/StatsPages/ShaderCookerStatsPage.cpp
@@ -125,7 +125,7 @@ void FShaderCookerStats::Initialize(uint32 Index)
 		const EShaderPlatform Platform = static_cast<EShaderPlatform>(PlatformIndex);
 
 		const FName ShaderFormatName = FDataDrivenShaderPlatformInfo::IsValid(Platform)
-			? ShaderPlatformToShaderFormatName(Platform) : NAME_None;
+			? FDataDrivenShaderPlatformInfo::GetShaderFormat(Platform) : NAME_None;
 
 		if (ShaderFormatName != NAME_None)
 		{
