Remove impossible cast in UStaticMeshLightingInfo::UpdateNames
#rnx

diff --git a/Engine/Source/Editor/StatsViewer/Private/StatsEntries/StaticMeshLightingInfo.cpp b/Engine/Source/Editor/StatsViewer/Private/StatsEntries/StaticMeshLightingInfo.cpp
--- a/Engine/Source/Editor/StatsViewer/Private/StatsEntries/StaticMeshLightingInfo.cpp
+++ b/Engine/Source/Editor/StatsViewer/Private/StatsEntries/StaticMeshLightingInfo.cpp
@@ -5,7 +5,6 @@
 #include "Engine/Level.h"
 #include "GameFramework/Actor.h"
 #include "Model.h"
-#include "Lightmass/LightmappedSurfaceCollection.h"
 
 #define LOCTEXT_NAMESPACE "Editor.StatsViewer"
 
@@ -18,21 +17,11 @@ void UStaticMeshLightingInfo::UpdateNames()
 {
 	if ( LevelName.Len() == 0 || StaticMeshActor.IsValid() )
 	{
-		AActor* Actor = Cast<AActor>(StaticMeshActor.Get());
-		ULightmappedSurfaceCollection* SurfaceCollection = Cast<ULightmappedSurfaceCollection>(StaticMeshActor.Get());
-
-		if (SurfaceCollection && SurfaceCollection->SourceModel)
-		{
-			LevelName = SurfaceCollection->SourceModel->GetOutermost()->GetName();
-		}
-		else if (Actor)
+		const AActor* Actor = StaticMeshActor.Get();
+		if (Actor)
 		{
 			LevelName = Actor->GetLevel()->GetOutermost()->GetName();
 		}
-		else if (StaticMeshActor.Get())
-		{
-			LevelName = StaticMeshActor.Get()->GetOutermost()->GetName();
-		}
 		else
 		{
 			LevelName = TEXT("<None>");
