Removed some FText->FString->FText churn

#rnx

diff --git a/Engine/Source/Editor/StatsViewer/Private/SStatsViewer.cpp b/Engine/Source/Editor/StatsViewer/Private/SStatsViewer.cpp
--- a/Engine/Source/Editor/StatsViewer/Private/SStatsViewer.cpp
+++ b/Engine/Source/Editor/StatsViewer/Private/SStatsViewer.cpp
@@ -718,24 +718,24 @@ TSharedRef<SWidget> SStatsViewer::OnGetFilterMenuContent() const
 			TWeakFieldPtr< FProperty > Property = *PropertyIter;
 			if( Property->HasAnyPropertyFlags(CPF_AssetRegistrySearchable) )
 			{
-				FString FilterName = Property->GetDisplayNameText().ToString();
-				if( FilterName.Len() == 0 )
+				FText FilterName = Property->GetDisplayNameText();
+				if( FilterName.IsEmpty() )
 				{
-					FilterName = UEditorEngine::GetFriendlyName(Property.Get());
+					FilterName = FText::AsCultureInvariant(UEditorEngine::GetFriendlyName(Property.Get()));
 				}
 
-				FString FilterDesc = Property->GetToolTipText().ToString();
-				if( FilterDesc.Len() == 0 )
+				FText FilterDesc = Property->GetToolTipText();
+				if( FilterDesc.IsEmpty() )
 				{
-					FilterDesc = UEditorEngine::GetFriendlyName(Property.Get());
+					FilterDesc = FText::AsCultureInvariant(UEditorEngine::GetFriendlyName(Property.Get()));
 				}
 
 				FFormatNamedArguments Arguments;
-				Arguments.Add(TEXT("FilterName"), FText::FromString(FilterName));
-				Arguments.Add(TEXT("FilterDesc"), FText::FromString(FilterDesc));
+				Arguments.Add(TEXT("FilterName"), FilterName);
+				Arguments.Add(TEXT("FilterDesc"), FilterDesc);
 
 				MenuBuilder.AddMenuEntry( 
-					FText::FromString( FilterName ), 
+					FilterName, 
 					FText::Format( LOCTEXT( "FilterMenuEntry_Tooltip", "Search statistics by {FilterName}.\n{FilterDesc}" ), Arguments ), 
 					FSlateIcon(), 
 					FUIAction( 
