TFieldIterator iterates raw pointers - no need to make them weak while iterating

#rnx

diff --git a/Engine/Source/Editor/StatsViewer/Private/SStatsViewer.cpp b/Engine/Source/Editor/StatsViewer/Private/SStatsViewer.cpp
--- a/Engine/Source/Editor/StatsViewer/Private/SStatsViewer.cpp
+++ b/Engine/Source/Editor/StatsViewer/Private/SStatsViewer.cpp
@@ -715,19 +715,19 @@ TSharedRef<SWidget> SStatsViewer::OnGetFilterMenuContent() const
 		int32 ColumnIndex = 0;
 		for (TFieldIterator<FProperty> PropertyIter( CurrentStats->GetEntryClass(), EFieldIteratorFlags::IncludeSuper ); PropertyIter; ++PropertyIter )
 		{
-			TWeakFieldPtr< FProperty > Property = *PropertyIter;
+			FProperty* Property = *PropertyIter;
 			if( Property->HasAnyPropertyFlags(CPF_AssetRegistrySearchable) )
 			{
 				FText FilterName = Property->GetDisplayNameText();
 				if( FilterName.IsEmpty() )
 				{
-					FilterName = FText::AsCultureInvariant(UEditorEngine::GetFriendlyName(Property.Get()));
+					FilterName = FText::AsCultureInvariant(UEditorEngine::GetFriendlyName(Property));
 				}
 
 				FText FilterDesc = Property->GetToolTipText();
 				if( FilterDesc.IsEmpty() )
 				{
-					FilterDesc = FText::AsCultureInvariant(UEditorEngine::GetFriendlyName(Property.Get()));
+					FilterDesc = FText::AsCultureInvariant(UEditorEngine::GetFriendlyName(Property));
 				}
 
 				FFormatNamedArguments Arguments;
