Update texture stats to distinguish between combined lod and not.

#rnx

diff --git a/Engine/Source/Editor/StatsViewer/Classes/TextureStats.h b/Engine/Source/Editor/StatsViewer/Classes/TextureStats.h
--- a/Engine/Source/Editor/StatsViewer/Classes/TextureStats.h
+++ b/Engine/Source/Editor/StatsViewer/Classes/TextureStats.h
@@ -67,9 +67,13 @@ class UTextureStats : public UObject
 	UPROPERTY(VisibleAnywhere, AssetRegistrySearchable, Category="Stats", meta=( ColumnWidth = "94" ) )
 	TEnumAsByte<TextureGroup> Group;
 
-	/** LOD Bias for this texture. (Texture LODBias + Texture group) */
-	UPROPERTY(VisibleAnywhere, AssetRegistrySearchable, Category="Stats", meta=( DisplayName = "LODBias", ColumnWidth = "70" ) )
-	int32 LODBias;
+	/** Combined LOD Bias for this texture. (Texture LODBias + Texture group) */
+	UPROPERTY(VisibleAnywhere, AssetRegistrySearchable, Category="Stats", meta=( DisplayName = "CombinedLODBias", ColumnWidth = "50" ) )
+	int32 CombinedLODBias;
+
+	/** LOD Bias for this texture. */
+	UPROPERTY(VisibleAnywhere, AssetRegistrySearchable, Category="Stats", meta=( DisplayName = "LODBias", ColumnWidth = "50" ) )
+	int32 TextureLODBias;
 
 	/** The memory used currently in KB */
 	UPROPERTY(VisibleAnywhere, AssetRegistrySearchable, Category="Stats", meta=( DisplayName = "Current Memory", ColumnWidth = "80", ShowTotal = "true", Unit = "KB" ) )
diff --git a/Engine/Source/Editor/StatsViewer/Private/StatsPages/TextureStatsPage.cpp b/Engine/Source/Editor/StatsViewer/Private/StatsPages/TextureStatsPage.cpp
--- a/Engine/Source/Editor/StatsViewer/Private/StatsPages/TextureStatsPage.cpp
+++ b/Engine/Source/Editor/StatsViewer/Private/StatsPages/TextureStatsPage.cpp
@@ -233,7 +233,8 @@ struct TextureStatsGenerator : public FFindReferencedAssets
 				Entry->FullyLoadedKB = InTexture->CalcTextureMemorySizeEnum( TMC_AllMipsBiased ) / 1024.0f;
 			}
 
-			Entry->LODBias = InTexture->GetCachedLODBias();
+			Entry->CombinedLODBias = InTexture->GetCachedLODBias();
+			Entry->TextureLODBias = InTexture->LODBias;
 
 			const FTexture* Resource = InTexture->GetResource();
 			if(Resource)
@@ -255,8 +256,8 @@ struct TextureStatsGenerator : public FFindReferencedAssets
 				Entry->CurrentDim.Y = Texture2D->GetSizeY() >> DroppedMips;
 
 				// Calculate the max dimensions
-				Entry->MaxDim.X = Texture2D->GetSizeX() >> Entry->LODBias;
-				Entry->MaxDim.Y = Texture2D->GetSizeY() >> Entry->LODBias;
+				Entry->MaxDim.X = Texture2D->GetSizeX() >> Entry->CombinedLODBias;
+				Entry->MaxDim.Y = Texture2D->GetSizeY() >> Entry->CombinedLODBias;
 			}
 			else
 			{
@@ -276,16 +277,16 @@ struct TextureStatsGenerator : public FFindReferencedAssets
 					Entry->CurrentDim.Y = Face->GetSizeY() >> DroppedMips;
 
 					// Calculate the max dimensions
-					Entry->MaxDim.X = Face->GetSizeX() >> Entry->LODBias;
-					Entry->MaxDim.Y = Face->GetSizeY() >> Entry->LODBias;
+					Entry->MaxDim.X = Face->GetSizeX() >> Entry->CombinedLODBias;
+					Entry->MaxDim.Y = Face->GetSizeY() >> Entry->CombinedLODBias;
 #else
 					// Calculate in game current dimensions 
-					Entry->CurrentDim.X = TextureCube->GetSizeX() >> Entry->LODBias;
-					Entry->CurrentDim.Y = TextureCube->GetSizeY() >> Entry->LODBias;
+					Entry->CurrentDim.X = TextureCube->GetSizeX() >> Entry->CombinedLODBias;
+					Entry->CurrentDim.Y = TextureCube->GetSizeY() >> Entry->CombinedLODBias;
 
 					// Calculate the max dimensions
-					Entry->MaxDim.X = TextureCube->GetSizeX() >> Entry->LODBias;
-					Entry->MaxDim.Y = TextureCube->GetSizeY() >> Entry->LODBias;
+					Entry->MaxDim.X = TextureCube->GetSizeX() >> Entry->CombinedLODBias;
+					Entry->MaxDim.Y = TextureCube->GetSizeY() >> Entry->CombinedLODBias;
 #endif
 				}
 			}
