Rename/move to pinball

#rnx

diff --git a/Engine/Source/Editor/StatsViewer/Private/SStatsViewer.cpp b/Engine/Source/Editor/StatsViewer/Private/SStatsViewer.cpp
--- a/Engine/Source/Editor/StatsViewer/Private/SStatsViewer.cpp
+++ b/Engine/Source/Editor/StatsViewer/Private/SStatsViewer.cpp
@@ -221,7 +221,8 @@ void SStatsViewer::Construct( const FArguments& InArgs )
 		{
 			InitialStatsPage = StatsPageManager.GetPage(FName(*DisplayedStatsPageName));
 		}
-		else
+
+		if(!InitialStatsPage.IsValid())
 		{
 			// Default to primitive stats if no config data exists yet
 			InitialStatsPage = StatsPageManager.GetPage(EStatsPage::PrimitiveStats);
