[Shaders] Mitigate crash in Shader Cook Statistics page.

This fixes a crash when the shader cook statistics sets of FShaderCookerStatsPage is not refreshed before accessing the entries again.
Use weak references to UObjects that can be checked if they have become stale and also add ensure condition when they do become stale and need re-initialization.

#rnx

diff --git a/Engine/Source/Editor/StatsViewer/Private/StatsPages/ShaderCookerStatsPage.cpp b/Engine/Source/Editor/StatsViewer/Private/StatsPages/ShaderCookerStatsPage.cpp
--- a/Engine/Source/Editor/StatsViewer/Private/StatsPages/ShaderCookerStatsPage.cpp
+++ b/Engine/Source/Editor/StatsViewer/Private/StatsPages/ShaderCookerStatsPage.cpp
@@ -34,7 +34,7 @@ public:
 	{
 	public:
 		FString Name;
-		TArray<UShaderCookerStats*> Stats;
+		TArray<TWeakObjectPtr<UShaderCookerStats>> Stats;
 		bool bInitialized;
 
 	};
@@ -79,14 +79,20 @@ public:
 		}
 	}
 
-	const TArray<UShaderCookerStats*>& GetShaderCookerStats(uint32 Index)
+	const TArray<TWeakObjectPtr<UShaderCookerStats>>& GetShaderCookerStats(uint32 Index)
 	{
-		FShaderCookerStatsSet Set = StatSets[Index];
+		check(Index < (uint32)StatSets.Num());
+		FShaderCookerStatsSet& Set = StatSets[Index];
 		if(!Set.bInitialized)
 		{
 			Initialize(Index);
 		}
-		return StatSets[Index].Stats;
+		else if (!ensureMsgf(Set.Stats.IsEmpty() || Set.Stats[0].IsValid(), TEXT("Shader cook statistics set got stale. This should not happen and indicates an issue with refreshing FShaderCookerStatsPage!")))
+		{
+			Set.Stats.Empty();
+			Initialize(Index);
+		}
+		return Set.Stats;
 	}
 	uint32 NumSets()
 	{
@@ -405,12 +411,15 @@ void FShaderCookerStatsPage::Generate( TArray< TWeakObjectPtr<UObject> >& OutObj
 	FShaderCookerStats& Stats = FShaderCookerStats::Get();
 	if((uint32)SelectedPlatform < Stats.NumSets())
 	{
-		const TArray<UShaderCookerStats*>& CookStats = Stats.GetShaderCookerStats(SelectedPlatform);
-		for(UShaderCookerStats* Stat: CookStats)
+		const TArray<TWeakObjectPtr<UShaderCookerStats>>& CookStats = Stats.GetShaderCookerStats(SelectedPlatform);
+		for(const TWeakObjectPtr<UShaderCookerStats>& Stat : CookStats)
 		{
-			if(0 == ObjectSetIndex || ObjectSetIndex == Stat->CategoryIndex)
+			if (Stat.IsValid())
 			{
-				OutObjects.Add(Stat);
+				if(0 == ObjectSetIndex || ObjectSetIndex == Stat->CategoryIndex)
+				{
+					OutObjects.Add(Stat);
+				}
 			}
 		}
 	}
