Optimize select actors in statistics tab
[RN] Optimize select actors in statistics tab
#Synced-CL 50645574

diff --git a/Engine/Source/Editor/StatsViewer/Private/ActorArrayHyperlinkColumn.cpp b/Engine/Source/Editor/StatsViewer/Private/ActorArrayHyperlinkColumn.cpp
--- a/Engine/Source/Editor/StatsViewer/Private/ActorArrayHyperlinkColumn.cpp
+++ b/Engine/Source/Editor/StatsViewer/Private/ActorArrayHyperlinkColumn.cpp
@@ -12,6 +12,7 @@
 #include "IPropertyTableCell.h"
 #include "PropertyPath.h"
 #include "ScopedTransaction.h"
+#include "Selection.h"
 #include "Widgets/Input/SHyperlink.h"
 
 #define LOCTEXT_NAMESPACE "Editor.StatsViewer"
@@ -78,7 +79,8 @@ public:
 			TArray<UObject*> ObjectsToSync;
 
 			const FScopedTransaction Transaction( LOCTEXT("SelectActors", "Statistics Select Actors") );
-			GEditor->SelectNone(false, false);
+			GEditor->GetSelectedActors()->BeginBatchSelectOperation();
+			GEditor->SelectNone(/*bNotify*/false, /*bDeselectBSPSurfs*/true);
 			for(uint32 ElementIndex = 0; ElementIndex < NumElements; ElementIndex++)
 			{
 				TSharedRef<IPropertyHandle> Element = Actors->GetElement(ElementIndex);
@@ -90,7 +92,7 @@ public:
 
 					if (Actor != nullptr)
 					{
-						GEditor->SelectActor(Actor, true, true, true);
+						GEditor->SelectActor(Actor, /*bInSelected*/true, /*bNotify*/false, /*bSelectEvenIfHidden*/true);
 						ActorsToFocus.Add(Actor);
 					}
 					else
@@ -100,6 +102,8 @@ public:
 				}
 			}
 
+			GEditor->GetSelectedActors()->EndBatchSelectOperation(/*bNotify*/true);
+
 			if(ActorsToFocus.Num() > 0)
 			{
 				GEditor->MoveViewportCamerasToActor(ActorsToFocus, false);	
